Fixed multiple missing virtual destructor compile errors appearing on Linux.

diff --git a/Engine/Source/Developer/AutomationDriver/Private/AutomatedApplication.cpp b/Engine/Source/Developer/AutomationDriver/Private/AutomatedApplication.cpp
--- a/Engine/Source/Developer/AutomationDriver/Private/AutomatedApplication.cpp
+++ b/Engine/Source/Developer/AutomationDriver/Private/AutomatedApplication.cpp
@@ -8,6 +8,9 @@ class FAutomatedCursor
 {
 public:
 
+	virtual ~FAutomatedCursor()
+	{ }
+
 	virtual FVector2D GetPosition() const override
 	{
 		if (bAllowMessageHandling)
diff --git a/Engine/Source/Developer/AutomationDriver/Private/AutomationDriver.cpp b/Engine/Source/Developer/AutomationDriver/Private/AutomationDriver.cpp
--- a/Engine/Source/Developer/AutomationDriver/Private/AutomationDriver.cpp
+++ b/Engine/Source/Developer/AutomationDriver/Private/AutomationDriver.cpp
@@ -274,6 +274,10 @@ TSharedRef<FAsyncAutomationDriver, ESPMode::ThreadSafe> FAsyncAutomationDriverFa
 }
 
 
+FAutomationDriver::~FAutomationDriver()
+{
+
+}
 
 bool FAutomationDriver::Wait(FTimespan Timespan)
 {
diff --git a/Engine/Source/Developer/AutomationDriver/Private/AutomationDriver.h b/Engine/Source/Developer/AutomationDriver/Private/AutomationDriver.h
--- a/Engine/Source/Developer/AutomationDriver/Private/AutomationDriver.h
+++ b/Engine/Source/Developer/AutomationDriver/Private/AutomationDriver.h
@@ -97,6 +97,8 @@ class FAutomationDriver
 {
 public:
 
+	virtual ~FAutomationDriver();
+
 	virtual bool Wait(FTimespan Timespan) override;
 	virtual bool Wait(const FDriverWaitDelegate& Delegate) override;
 
diff --git a/Engine/Source/Developer/AutomationDriver/Private/DriverElement.cpp b/Engine/Source/Developer/AutomationDriver/Private/DriverElement.cpp
--- a/Engine/Source/Developer/AutomationDriver/Private/DriverElement.cpp
+++ b/Engine/Source/Developer/AutomationDriver/Private/DriverElement.cpp
@@ -831,6 +831,9 @@ class FDriverElementCollection
 {
 public:
 
+	virtual ~FDriverElementCollection()
+	{ }
+
 	virtual TArray<TSharedRef<IDriverElement, ESPMode::ThreadSafe>> GetElements() override
 	{
 		if (IsInGameThread())
@@ -905,6 +908,9 @@ class FDriverElement
 {
 public:
 
+	virtual ~FDriverElement()
+	{ }
+
 	virtual bool Hover() override
 	{
 		TSharedRef<IDriverSequence, ESPMode::ThreadSafe> Sequence = Driver->CreateSequence();
diff --git a/Engine/Source/Developer/AutomationDriver/Private/DriverSequence.cpp b/Engine/Source/Developer/AutomationDriver/Private/DriverSequence.cpp
--- a/Engine/Source/Developer/AutomationDriver/Private/DriverSequence.cpp
+++ b/Engine/Source/Developer/AutomationDriver/Private/DriverSequence.cpp
@@ -1362,6 +1362,9 @@ class FAsyncDriverSequence
 {
 public:
 
+	virtual ~FAsyncDriverSequence()
+	{ }
+
 	virtual IAsyncActionSequence& Actions()
 	{
 		return *ActionSequence;
@@ -1866,6 +1869,9 @@ class FDriverSequence
 {
 public:
 
+	virtual ~FDriverSequence()
+	{ }
+
 	virtual IActionSequence& Actions()
 	{
 		return *ActionSequence;
diff --git a/Engine/Source/Developer/AutomationDriver/Private/Locators/SlateWidgetLocatorByDelegate.cpp b/Engine/Source/Developer/AutomationDriver/Private/Locators/SlateWidgetLocatorByDelegate.cpp
--- a/Engine/Source/Developer/AutomationDriver/Private/Locators/SlateWidgetLocatorByDelegate.cpp
+++ b/Engine/Source/Developer/AutomationDriver/Private/Locators/SlateWidgetLocatorByDelegate.cpp
@@ -9,6 +9,9 @@ class FSlateWidgetLocatorByWidgetDelegate
 {
 public:
 
+	virtual ~FSlateWidgetLocatorByWidgetDelegate()
+	{ }
+
 	virtual FString ToDebugString() const
 	{
 		FString DelegateName;
@@ -62,6 +65,9 @@ class FSlateWidgetLocatorByWidgetPathDelegate
 {
 public:
 
+	virtual ~FSlateWidgetLocatorByWidgetPathDelegate()
+	{ }
+
 	virtual FString ToDebugString() const
 	{
 		FString DelegateName;
diff --git a/Engine/Source/Developer/AutomationDriver/Private/Locators/SlateWidgetLocatorByPath.cpp b/Engine/Source/Developer/AutomationDriver/Private/Locators/SlateWidgetLocatorByPath.cpp
--- a/Engine/Source/Developer/AutomationDriver/Private/Locators/SlateWidgetLocatorByPath.cpp
+++ b/Engine/Source/Developer/AutomationDriver/Private/Locators/SlateWidgetLocatorByPath.cpp
@@ -25,6 +25,9 @@ private:
 	{
 	public:
 
+		virtual ~FIdMatcher()
+		{ }
+
 		virtual bool IsMatch(const TSharedRef<SWidget>& Widget) const override
 		{
 			const TArray<TSharedRef<FDriverIdMetaData>> AllIdMetaData = Widget->GetAllMetaData<FDriverIdMetaData>();
@@ -57,6 +60,9 @@ private:
 	{
 	public:
 
+		virtual ~FTagMatcher()
+		{ }
+
 		virtual bool IsMatch(const TSharedRef<SWidget>& Widget) const override
 		{
 			if (Widget->GetTag() == Tag)
@@ -94,6 +100,9 @@ private:
 	{
 	public:
 
+		virtual ~FTypeMatcher()
+		{ }
+
 		virtual bool IsMatch(const TSharedRef<SWidget>& Widget) const override
 		{
 			return Widget->GetType() == Type;
diff --git a/Engine/Source/Developer/AutomationDriver/Private/Locators/SlateWidgetLocatorByUniqueTag.cpp b/Engine/Source/Developer/AutomationDriver/Private/Locators/SlateWidgetLocatorByUniqueTag.cpp
--- a/Engine/Source/Developer/AutomationDriver/Private/Locators/SlateWidgetLocatorByUniqueTag.cpp
+++ b/Engine/Source/Developer/AutomationDriver/Private/Locators/SlateWidgetLocatorByUniqueTag.cpp
@@ -10,6 +10,9 @@ class FSlateWidgetLocatorByUniqueTag
 {
 public:
 
+	virtual ~FSlateWidgetLocatorByUniqueTag()
+	{ }
+
 	virtual FString ToDebugString() const
 	{
 		const FString Address = FString::Printf(TEXT("0x%p"), &UniqueMetaData.Get());
diff --git a/Engine/Source/Developer/AutomationDriver/Private/PassThroughMessageHandler .cpp b/Engine/Source/Developer/AutomationDriver/Private/PassThroughMessageHandler .cpp
--- a/Engine/Source/Developer/AutomationDriver/Private/PassThroughMessageHandler .cpp	
+++ b/Engine/Source/Developer/AutomationDriver/Private/PassThroughMessageHandler .cpp	
@@ -9,6 +9,9 @@ class FPassThroughMessageHandlerImpl
 {
 public:
 
+	virtual ~FPassThroughMessageHandlerImpl()
+	{ }
+
 	virtual bool IsHandlingMessages() const override
 	{
 		return bAllowMessageHandling;
@@ -456,6 +459,9 @@ class FPassThroughMessageHandlerFactory
 {
 public:
 
+	virtual ~FPassThroughMessageHandlerFactory()
+	{ }
+
 	virtual TSharedRef<FPassThroughMessageHandler> Create(
 		const TSharedRef<FGenericApplicationMessageHandler>& MessageHandler) const override
 	{
diff --git a/Engine/Source/Developer/AutomationDriver/Private/PassThroughMessageHandler.h b/Engine/Source/Developer/AutomationDriver/Private/PassThroughMessageHandler.h
--- a/Engine/Source/Developer/AutomationDriver/Private/PassThroughMessageHandler.h
+++ b/Engine/Source/Developer/AutomationDriver/Private/PassThroughMessageHandler.h
@@ -8,6 +8,9 @@ class FPassThroughMessageHandler
 {
 public:
 
+	virtual ~FPassThroughMessageHandler()
+	{ }
+
 	virtual bool IsHandlingMessages() const = 0;
 
 	virtual void SetAllowMessageHandling(bool bValue) = 0;
@@ -18,6 +21,9 @@ class IPassThroughMessageHandlerFactory
 {
 public:
 
+	virtual ~IPassThroughMessageHandlerFactory()
+	{ }
+
 	virtual TSharedRef<FPassThroughMessageHandler> Create(
 		const TSharedRef<FGenericApplicationMessageHandler>& MessageHandler) const = 0;
 };
diff --git a/Engine/Source/Developer/AutomationDriver/Private/SlateWidgetElement.cpp b/Engine/Source/Developer/AutomationDriver/Private/SlateWidgetElement.cpp
--- a/Engine/Source/Developer/AutomationDriver/Private/SlateWidgetElement.cpp
+++ b/Engine/Source/Developer/AutomationDriver/Private/SlateWidgetElement.cpp
@@ -16,6 +16,9 @@ class FSlateWidgetElement
 {
 public:
 
+	virtual ~FSlateWidgetElement()
+	{ }
+
 	virtual FString ToDebugString() const
 	{
 		const TSharedRef<SWidget> Widget = WidgetPath.Widgets.Last().Widget;
diff --git a/Engine/Source/Developer/AutomationDriver/Private/Specs/AutomationDriverSpecSuiteViewModel.cpp b/Engine/Source/Developer/AutomationDriver/Private/Specs/AutomationDriverSpecSuiteViewModel.cpp
--- a/Engine/Source/Developer/AutomationDriver/Private/Specs/AutomationDriverSpecSuiteViewModel.cpp
+++ b/Engine/Source/Developer/AutomationDriver/Private/Specs/AutomationDriverSpecSuiteViewModel.cpp
@@ -8,6 +8,9 @@ class FAutomationDriverSpecSuiteViewModel
 {
 public:
 
+	virtual ~FAutomationDriverSpecSuiteViewModel()
+	{ }
+
 	virtual FText GetFormText(EFormElement Element) const override
 	{
 		return FormTextMap.FindOrAdd(Element);
diff --git a/Engine/Source/Developer/AutomationDriver/Private/Specs/AutomationDriverSpecSuiteViewModel.h b/Engine/Source/Developer/AutomationDriver/Private/Specs/AutomationDriverSpecSuiteViewModel.h
--- a/Engine/Source/Developer/AutomationDriver/Private/Specs/AutomationDriverSpecSuiteViewModel.h
+++ b/Engine/Source/Developer/AutomationDriver/Private/Specs/AutomationDriverSpecSuiteViewModel.h
@@ -134,6 +134,9 @@ class IAutomationDriverSpecSuiteViewModel
 {
 public:
 
+	virtual ~IAutomationDriverSpecSuiteViewModel()
+	{ }
+
 	/** @return The current value of the form text element */
 	virtual FText GetFormText(EFormElement Element) const = 0;
 
