Fixed a missing virtual destructor

diff --git a/Engine/Source/Developer/AutomationDriver/Private/Locators/SlateWidgetLocatorByPath.cpp b/Engine/Source/Developer/AutomationDriver/Private/Locators/SlateWidgetLocatorByPath.cpp
--- a/Engine/Source/Developer/AutomationDriver/Private/Locators/SlateWidgetLocatorByPath.cpp
+++ b/Engine/Source/Developer/AutomationDriver/Private/Locators/SlateWidgetLocatorByPath.cpp
@@ -117,6 +117,9 @@ private:
 
 public:
 
+	virtual ~FSlateWidgetLocatorByPath()
+	{ }
+
 	virtual FString ToDebugString() const
 	{
 		return TEXT("[By::Path] ") + Path;
