Added the API macro declarations to all the public static members of the AutomationDriver.
Moved the DriverIdMetaData out of the AutomationDriver module so it is no longer a dependency to use the tags.
The SlateWidgetElement now uses the DrawSize instead of the LocalSize in order to take into account UI scaling.

diff --git a/Engine/Source/Developer/AutomationDriver/Private/DriverMetaData.cpp b/Engine/Source/Developer/AutomationDriver/Private/DriverMetaData.cpp
deleted file mode 100644
diff --git a/Engine/Source/Developer/AutomationDriver/Private/MetaData/DriverIdMetaData.h b/Engine/Source/Developer/AutomationDriver/Private/MetaData/DriverIdMetaData.h
deleted file mode 100644
diff --git a/Engine/Source/Developer/AutomationDriver/Private/SlateWidgetElement.cpp b/Engine/Source/Developer/AutomationDriver/Private/SlateWidgetElement.cpp
--- a/Engine/Source/Developer/AutomationDriver/Private/SlateWidgetElement.cpp
+++ b/Engine/Source/Developer/AutomationDriver/Private/SlateWidgetElement.cpp
@@ -83,7 +83,7 @@ public:
 
 	virtual FVector2D GetSize() const override
 	{
-		return WidgetPath.Widgets.Last().Geometry.GetLocalSize();
+		return WidgetPath.Widgets.Last().Geometry.GetDrawSize();
 	}
 
 	virtual TSharedPtr<FGenericWindow> GetWindow() const override
diff --git a/Engine/Source/Developer/AutomationDriver/Private/Specs/AutomationDriver.spec.cpp b/Engine/Source/Developer/AutomationDriver/Private/Specs/AutomationDriver.spec.cpp
--- a/Engine/Source/Developer/AutomationDriver/Private/Specs/AutomationDriver.spec.cpp
+++ b/Engine/Source/Developer/AutomationDriver/Private/Specs/AutomationDriver.spec.cpp
@@ -903,6 +903,7 @@ void FAutomationDriverSpec::Define()
 		{
 			It("should return true if the element is currently visible in the SWidget DOM being displayed", EAsyncExecution::ThreadPool, [this]()
 			{
+				Driver->Wait(Until::ElementIsVisible(By::Id("Piano"), FWaitTimeout::InSeconds(1)));
 				TEST_TRUE(Driver->FindElement(By::Id("Piano"))->IsVisible());
 			});
 
diff --git a/Engine/Source/Developer/AutomationDriver/Public/DriverMetaData.h b/Engine/Source/Developer/AutomationDriver/Public/DriverMetaData.h
deleted file mode 100644
diff --git a/Engine/Source/Developer/AutomationDriver/Public/LocateBy.h b/Engine/Source/Developer/AutomationDriver/Public/LocateBy.h
--- a/Engine/Source/Developer/AutomationDriver/Public/LocateBy.h
+++ b/Engine/Source/Developer/AutomationDriver/Public/LocateBy.h
@@ -9,7 +9,7 @@ DECLARE_DELEGATE_OneParam(FLocateSlateWidgetPathElementDelegate, TArray<FWidgetP
 /**
  * Represents a collection of fluent helper functions designed to make accessing and creating element locators easy
  */
-class By
+class AUTOMATIONDRIVER_API By
 {
 public:
 
diff --git a/Engine/Source/Developer/AutomationDriver/Public/WaitUntil.h b/Engine/Source/Developer/AutomationDriver/Public/WaitUntil.h
--- a/Engine/Source/Developer/AutomationDriver/Public/WaitUntil.h
+++ b/Engine/Source/Developer/AutomationDriver/Public/WaitUntil.h
@@ -6,7 +6,7 @@ class IElementLocator;
 /**
  * Represents the state of an active wait action for the driver
  */
-struct FDriverWaitResponse
+struct AUTOMATIONDRIVER_API FDriverWaitResponse
 {
 public:
 
@@ -52,7 +52,7 @@ DECLARE_DELEGATE_RetVal(bool, FDriverWaitConditionDelegate);
 /**
  * A fluent wrapper around timespan to enforce obvious differences between specified Timeout and Interval values for waits
  */
-class FWaitTimeout
+class AUTOMATIONDRIVER_API FWaitTimeout
 {
 public:
 
@@ -69,7 +69,7 @@ public:
 /**
  * A fluent wrapper around timespan to enforce obvious differences between specified Timeout and Interval values for waits
  */
-class FWaitInterval
+class AUTOMATIONDRIVER_API FWaitInterval
 {
 public:
 
@@ -86,7 +86,7 @@ public:
 /**
  * Represents a collection of fluent helper functions designed to make accessing and creating driver wait delegates easier
  */
-class Until
+class AUTOMATIONDRIVER_API Until
 {
 public:
 
diff --git a/Engine/Source/Runtime/Slate/Private/Framework/MetaData/DriverMetaData.cpp b/Engine/Source/Runtime/Slate/Private/Framework/MetaData/DriverMetaData.cpp
new file mode 100644
--- /dev/null
+++ b/Engine/Source/Runtime/Slate/Private/Framework/MetaData/DriverMetaData.cpp
@@ -0,0 +1,9 @@
+// Copyright 1998-2016 Epic Games, Inc. All Rights Reserved.
+
+#include "DriverMetaData.h"
+#include "DriverIdMetaData.h"
+
+TSharedRef<ISlateMetaData> FDriverMetaData::Id(FName InTag)
+{
+	return MakeShareable(new FDriverIdMetaData(MoveTemp(InTag)));
+}
diff --git a/Engine/Source/Runtime/Slate/Public/Framework/MetaData/DriverIdMetaData.h b/Engine/Source/Runtime/Slate/Public/Framework/MetaData/DriverIdMetaData.h
new file mode 100644
--- /dev/null
+++ b/Engine/Source/Runtime/Slate/Public/Framework/MetaData/DriverIdMetaData.h
@@ -0,0 +1,17 @@
+// Copyright 1998-2016 Epic Games, Inc. All Rights Reserved.
+#pragma once
+
+#include "ISlateMetaData.h"
+
+class FDriverIdMetaData
+	: public ISlateMetaData
+{
+public:
+	SLATE_METADATA_TYPE(FDriverIdMetaData, ISlateMetaData)
+
+	FDriverIdMetaData(FName InId)
+		: Id(MoveTemp(InId))
+	{ }
+
+	FName Id;
+};
diff --git a/Engine/Source/Runtime/Slate/Public/Framework/MetaData/DriverMetaData.h b/Engine/Source/Runtime/Slate/Public/Framework/MetaData/DriverMetaData.h
new file mode 100644
--- /dev/null
+++ b/Engine/Source/Runtime/Slate/Public/Framework/MetaData/DriverMetaData.h
@@ -0,0 +1,17 @@
+// Copyright 1998-2016 Epic Games, Inc. All Rights Reserved.
+#pragma once
+
+class ISlateMetaData;
+
+/**
+ * A static helper class which is used to easily construct various types of AutomationDriver specific SlateMetaData.
+ */
+class SLATE_API FDriverMetaData
+{
+public:
+
+	/**
+	 * @return the automation driver specific metadata type for specifying an Id for a widget
+	 */
+	static TSharedRef<ISlateMetaData> Id(FName Tag);
+};
\ No newline at end of file
