Fix for TSharedFromThis mode discrepancy.

diff --git a/Engine/Source/Developer/AutomationDriver/Private/DriverSequence.cpp b/Engine/Source/Developer/AutomationDriver/Private/DriverSequence.cpp
--- a/Engine/Source/Developer/AutomationDriver/Private/DriverSequence.cpp
+++ b/Engine/Source/Developer/AutomationDriver/Private/DriverSequence.cpp
@@ -178,7 +178,7 @@ public:
 
 class FAsyncActionSequence
 	: public IAsyncActionSequence
-	, public TSharedFromThis<FAsyncActionSequence>
+	, public TSharedFromThis<FAsyncActionSequence, ESPMode::ThreadSafe>
 {
 public:
 
