Fix automation driver module message handling

diff --git a/Engine/Source/Developer/AutomationDriver/Private/AutomationDriverModule.cpp b/Engine/Source/Developer/AutomationDriver/Private/AutomationDriverModule.cpp
--- a/Engine/Source/Developer/AutomationDriver/Private/AutomationDriverModule.cpp
+++ b/Engine/Source/Developer/AutomationDriver/Private/AutomationDriverModule.cpp
@@ -60,6 +60,11 @@ public:
 			RealApplication.ToSharedRef(),
 			FPassThroughMessageHandlerFactoryFactory::Create());
 
+		if (AutomatedApplication.IsValid())
+		{
+			AutomatedApplication->AllowPlatformMessageHandling();
+		}
+		
 		FSlateApplication::Get().SetPlatformApplication(AutomatedApplication.ToSharedRef());
 	}
 
@@ -70,6 +75,8 @@ public:
 			return;
 		}
 
+		AutomatedApplication->DisablePlatformMessageHandling();
+
 		FSlateApplication::Get().SetPlatformApplication(RealApplication.ToSharedRef());
 		RealApplication->SetMessageHandler(RealMessageHandler.ToSharedRef());
 
