Update PassThroughMessageHandler to have the up to date OnController functions

#rnx

diff --git a/Engine/Source/Developer/AutomationDriver/Private/PassThroughMessageHandler.cpp b/Engine/Source/Developer/AutomationDriver/Private/PassThroughMessageHandler.cpp
--- a/Engine/Source/Developer/AutomationDriver/Private/PassThroughMessageHandler.cpp
+++ b/Engine/Source/Developer/AutomationDriver/Private/PassThroughMessageHandler.cpp
@@ -178,6 +178,7 @@ public:
 		return RealMessageHandler->OnCursorSet();
 	}
 
+	UE_DEPRECATED(5.1, "This version of OnControllerAnalog is deprecated, please use the version that takes an FPlatformUserId instead")
 	virtual bool OnControllerAnalog(FGamepadKeyNames::Type KeyName, int32 ControllerId, float AnalogValue) override
 	{
 		if (!bAllowMessageHandling)
@@ -185,9 +186,12 @@ public:
 			return false;
 		}
 
+		PRAGMA_DISABLE_DEPRECATION_WARNINGS
 		return RealMessageHandler->OnControllerAnalog(KeyName, ControllerId, AnalogValue);
+		PRAGMA_ENABLE_DEPRECATION_WARNINGS
 	}
 
+	UE_DEPRECATED(5.1, "This version of OnControllerAnalog is deprecated, please use the version that takes an FPlatformUserId instead")
 	virtual bool OnControllerButtonPressed(FGamepadKeyNames::Type KeyName, int32 ControllerId, bool IsRepeat) override
 	{
 		if (!bAllowMessageHandling)
@@ -195,9 +199,12 @@ public:
 			return false;
 		}
 
+		PRAGMA_DISABLE_DEPRECATION_WARNINGS
 		return RealMessageHandler->OnControllerButtonPressed(KeyName, ControllerId, IsRepeat);
+		PRAGMA_ENABLE_DEPRECATION_WARNINGS
 	}
 
+	UE_DEPRECATED(5.1, "This version of OnControllerAnalog is deprecated, please use the version that takes an FPlatformUserId instead")
 	virtual bool OnControllerButtonReleased(FGamepadKeyNames::Type KeyName, int32 ControllerId, bool IsRepeat) override
 	{
 		if (!bAllowMessageHandling)
@@ -205,7 +212,39 @@ public:
 			return false;
 		}
 
+		PRAGMA_DISABLE_DEPRECATION_WARNINGS
 		return RealMessageHandler->OnControllerButtonReleased(KeyName, ControllerId, IsRepeat);
+		PRAGMA_ENABLE_DEPRECATION_WARNINGS
+	}
+	
+	virtual bool OnControllerAnalog(FGamepadKeyNames::Type KeyName, FPlatformUserId PlatformUserId, FInputDeviceId InputDeviceId, float AnalogValue)
+	{
+		if (!bAllowMessageHandling)
+		{
+			return false;
+		}
+
+		return RealMessageHandler->OnControllerAnalog(KeyName, PlatformUserId, InputDeviceId, AnalogValue);
+	}
+
+	virtual bool OnControllerButtonPressed(FGamepadKeyNames::Type KeyName, FPlatformUserId PlatformUserId, FInputDeviceId InputDeviceId, bool IsRepeat)
+	{
+		if (!bAllowMessageHandling)
+		{
+			return false;
+		}
+
+		return RealMessageHandler->OnControllerButtonPressed(KeyName, PlatformUserId, InputDeviceId, IsRepeat);
+	}
+
+	virtual bool OnControllerButtonReleased(FGamepadKeyNames::Type KeyName, FPlatformUserId PlatformUserId, FInputDeviceId InputDeviceId, bool IsRepeat)
+	{
+		if (!bAllowMessageHandling)
+		{
+			return false;
+		}
+
+		return RealMessageHandler->OnControllerButtonReleased(KeyName, PlatformUserId, InputDeviceId, IsRepeat);
 	}
 
 	virtual void OnBeginGesture() override
