Unregister user on accessibility message handler on AutomationDriver teardown

#rnx

diff --git a/Engine/Source/Developer/AutomationDriver/Private/AutomationDriverModule.cpp b/Engine/Source/Developer/AutomationDriver/Private/AutomationDriverModule.cpp
--- a/Engine/Source/Developer/AutomationDriver/Private/AutomationDriverModule.cpp
+++ b/Engine/Source/Developer/AutomationDriver/Private/AutomationDriverModule.cpp
@@ -64,7 +64,7 @@ public:
 		{
 			AutomatedApplication->AllowPlatformMessageHandling();
 		}
-		
+
 		FSlateApplication::Get().SetPlatformApplication(AutomatedApplication.ToSharedRef());
 	}
 
@@ -77,6 +77,12 @@ public:
 
 		AutomatedApplication->DisablePlatformMessageHandling();
 
+#if WITH_ACCESSIBILITY
+		// Unregister primary user on the accessible handler before swapping back the original application
+		FGenericAccessibleUserRegistry& UserRegistry = AutomatedApplication->GetAccessibleMessageHandler()->GetAccessibleUserRegistry();
+		UserRegistry.UnregisterUser(FGenericAccessibleUserRegistry::GetPrimaryUserIndex());
+#endif
+
 		FSlateApplication::Get().SetPlatformApplication(RealApplication.ToSharedRef());
 		RealApplication->SetMessageHandler(RealMessageHandler.ToSharedRef());
 
