AutomationDriver - use share pointer to store FAutomatedApplication
Fix crash related to GC
+ harden reset with some Latent commands after they timed out

diff --git a/Engine/Source/Developer/AutomationDriver/Private/AutomatedApplication.cpp b/Engine/Source/Developer/AutomationDriver/Private/AutomatedApplication.cpp
--- a/Engine/Source/Developer/AutomationDriver/Private/AutomatedApplication.cpp
+++ b/Engine/Source/Developer/AutomationDriver/Private/AutomatedApplication.cpp
@@ -201,6 +201,21 @@ public:
 		}
 	}
 
+	virtual bool IsHandlingMessages() override
+	{
+		if (!PassThroughMessageHandler.IsValid() || !PassThroughMessageHandler->IsHandlingMessages())
+		{
+			return false;
+		}
+
+		if (!AutomatedCursor.IsValid() || !AutomatedCursor->IsHandlingMessages())
+		{
+			return false;
+		}
+
+		return true;
+	}
+
 	virtual void SetMessageHandler(const TSharedRef<FGenericApplicationMessageHandler>& InMessageHandler) override
 	{
 		RealMessageHandler = InMessageHandler;
diff --git a/Engine/Source/Developer/AutomationDriver/Private/AutomatedApplication.h b/Engine/Source/Developer/AutomationDriver/Private/AutomatedApplication.h
--- a/Engine/Source/Developer/AutomationDriver/Private/AutomatedApplication.h
+++ b/Engine/Source/Developer/AutomationDriver/Private/AutomatedApplication.h
@@ -18,6 +18,7 @@ public:
 
 	virtual void AllowPlatformMessageHandling() = 0;
 	virtual void DisablePlatformMessageHandling() = 0;
+	virtual bool IsHandlingMessages() = 0;
 
 	virtual void SetFakeModifierKeys(FModifierKeysState Value) = 0;
 
diff --git a/Engine/Source/Developer/AutomationDriver/Private/AutomationDriver.cpp b/Engine/Source/Developer/AutomationDriver/Private/AutomationDriver.cpp
--- a/Engine/Source/Developer/AutomationDriver/Private/AutomationDriver.cpp
+++ b/Engine/Source/Developer/AutomationDriver/Private/AutomationDriver.cpp
@@ -255,7 +255,7 @@ void FAsyncAutomationDriver::Initialize()
 }
 
 TSharedRef<FAsyncAutomationDriver, ESPMode::ThreadSafe> FAsyncAutomationDriverFactory::Create(
-	FAutomatedApplication* const AutomatedApplication)
+	const TSharedRef<FAutomatedApplication, ESPMode::ThreadSafe>& AutomatedApplication)
 {
 	TSharedRef<FAsyncAutomationDriver, ESPMode::ThreadSafe> Instance = MakeShareable(new FAsyncAutomationDriver(
 		AutomatedApplication,
@@ -266,7 +266,7 @@ TSharedRef<FAsyncAutomationDriver, ESPMode::ThreadSafe> FAsyncAutomationDriverFa
 }
 
 TSharedRef<FAsyncAutomationDriver, ESPMode::ThreadSafe> FAsyncAutomationDriverFactory::Create(
-	FAutomatedApplication* const AutomatedApplication,
+	const TSharedRef<FAutomatedApplication, ESPMode::ThreadSafe>& AutomatedApplication,
 	const TSharedRef<FDriverConfiguration, ESPMode::ThreadSafe>& Configuration)
 {
 	return MakeShareable(new FAsyncAutomationDriver(
@@ -382,7 +382,7 @@ bool FAutomationDriver::IsPressed(EMouseButtons::Type Button) const
 }
 
 TSharedRef<FAutomationDriver, ESPMode::ThreadSafe> FAutomationDriverFactory::Create(
-	FAutomatedApplication* const AutomatedApplication)
+	const TSharedRef<FAutomatedApplication, ESPMode::ThreadSafe>& AutomatedApplication)
 {
 	return MakeShareable(new FAutomationDriver(
 		AutomatedApplication,
@@ -390,7 +390,7 @@ TSharedRef<FAutomationDriver, ESPMode::ThreadSafe> FAutomationDriverFactory::Cre
 }
 
 TSharedRef<FAutomationDriver, ESPMode::ThreadSafe> FAutomationDriverFactory::Create(
-	FAutomatedApplication* const AutomatedApplication,
+	const TSharedRef<FAutomatedApplication, ESPMode::ThreadSafe>& AutomatedApplication,
 	const TSharedRef<FDriverConfiguration, ESPMode::ThreadSafe>& Configuration)
 {
 	return MakeShareable(new FAutomationDriver(
diff --git a/Engine/Source/Developer/AutomationDriver/Private/AutomationDriver.h b/Engine/Source/Developer/AutomationDriver/Private/AutomationDriver.h
--- a/Engine/Source/Developer/AutomationDriver/Private/AutomationDriver.h
+++ b/Engine/Source/Developer/AutomationDriver/Private/AutomationDriver.h
@@ -55,7 +55,7 @@ public:
 private:
 
 	FAsyncAutomationDriver(
-		FAutomatedApplication* const InApplication,
+		const TSharedRef<FAutomatedApplication, ESPMode::ThreadSafe>& InApplication,
 		const TSharedRef<FDriverConfiguration, ESPMode::ThreadSafe>& InConfiguration)
 		: Application(InApplication)
 		, Configuration(InConfiguration)
@@ -65,7 +65,7 @@ private:
 
 private:
 
-	FAutomatedApplication* const Application;
+	const TSharedRef<FAutomatedApplication, ESPMode::ThreadSafe> Application;
 	const TSharedRef<FDriverConfiguration, ESPMode::ThreadSafe> Configuration;
 
 	TSet<FKey> PressedModifiers;
@@ -82,10 +82,10 @@ class FAsyncAutomationDriverFactory
 public:
 
 	static TSharedRef<FAsyncAutomationDriver, ESPMode::ThreadSafe> Create(
-		FAutomatedApplication* const AutomatedApplication);
+		const TSharedRef<FAutomatedApplication, ESPMode::ThreadSafe>& AutomatedApplication);
 
 	static TSharedRef<FAsyncAutomationDriver, ESPMode::ThreadSafe> Create(
-		FAutomatedApplication* const AutomatedApplication,
+		const TSharedRef<FAutomatedApplication, ESPMode::ThreadSafe>& AutomatedApplication,
 		const TSharedRef<FDriverConfiguration, ESPMode::ThreadSafe>& Configuration);
 };
 
@@ -129,7 +129,7 @@ public:
 private:
 
 	FAutomationDriver(
-		FAutomatedApplication* const InApplication,
+		const TSharedRef<FAutomatedApplication, ESPMode::ThreadSafe>& InApplication,
 		const TSharedRef<FAsyncAutomationDriver, ESPMode::ThreadSafe>& InAsyncDriver)
 		: Application(InApplication)
 		, AsyncDriver(InAsyncDriver)
@@ -137,7 +137,7 @@ private:
 
 private:
 
-	FAutomatedApplication* const Application;
+	const TSharedRef<FAutomatedApplication, ESPMode::ThreadSafe> Application;
 	const TSharedRef<FAsyncAutomationDriver, ESPMode::ThreadSafe> AsyncDriver;
 
 	friend FAutomationDriverFactory;
@@ -148,9 +148,9 @@ class FAutomationDriverFactory
 public:
 
 	static TSharedRef<FAutomationDriver, ESPMode::ThreadSafe> Create(
-		FAutomatedApplication* const AutomatedApplication);
+		const TSharedRef<FAutomatedApplication, ESPMode::ThreadSafe>& AutomatedApplication);
 
 	static TSharedRef<FAutomationDriver, ESPMode::ThreadSafe> Create(
-		FAutomatedApplication* const AutomatedApplication,
+		const TSharedRef<FAutomatedApplication, ESPMode::ThreadSafe>& AutomatedApplication,
 		const TSharedRef<FDriverConfiguration, ESPMode::ThreadSafe>& Configuration);
 };
diff --git a/Engine/Source/Developer/AutomationDriver/Private/AutomationDriverModule.cpp b/Engine/Source/Developer/AutomationDriver/Private/AutomationDriverModule.cpp
--- a/Engine/Source/Developer/AutomationDriver/Private/AutomationDriverModule.cpp
+++ b/Engine/Source/Developer/AutomationDriver/Private/AutomationDriverModule.cpp
@@ -23,22 +23,22 @@ public:
 
 	virtual TSharedRef<IAutomationDriver, ESPMode::ThreadSafe> CreateDriver() const override
 	{
-		return FAutomationDriverFactory::Create(AutomatedApplication.Get());
+		return FAutomationDriverFactory::Create(AutomatedApplication.ToSharedRef());
 	}
 
 	virtual TSharedRef<IAutomationDriver, ESPMode::ThreadSafe> CreateDriver(const TSharedRef<FDriverConfiguration, ESPMode::ThreadSafe>& Configuration) const override
 	{
-		return FAutomationDriverFactory::Create(AutomatedApplication.Get(), Configuration);
+		return FAutomationDriverFactory::Create(AutomatedApplication.ToSharedRef(), Configuration);
 	}
 
 	virtual TSharedRef<IAsyncAutomationDriver, ESPMode::ThreadSafe> CreateAsyncDriver() const override
 	{
-		return FAsyncAutomationDriverFactory::Create(AutomatedApplication.Get());
+		return FAsyncAutomationDriverFactory::Create(AutomatedApplication.ToSharedRef());
 	}
 
 	virtual TSharedRef<IAsyncAutomationDriver, ESPMode::ThreadSafe> CreateAsyncDriver(const TSharedRef<FDriverConfiguration, ESPMode::ThreadSafe>& Configuration) const override
 	{
-		return FAsyncAutomationDriverFactory::Create(AutomatedApplication.Get(), Configuration);
+		return FAsyncAutomationDriverFactory::Create(AutomatedApplication.ToSharedRef(), Configuration);
 	}
 
 	virtual bool IsEnabled() const override
@@ -48,7 +48,7 @@ public:
 
 	virtual void Enable() override
 	{
-		if (AutomatedApplication.IsValid())
+		if (IsEnabled())
 		{
 			return;
 		}
@@ -69,7 +69,7 @@ public:
 
 	virtual void Disable() override
 	{
-		if (!AutomatedApplication.IsValid())
+		if (!IsEnabled())
 		{
 			return;
 		}
diff --git a/Engine/Source/Developer/AutomationDriver/Private/DriverSequence.cpp b/Engine/Source/Developer/AutomationDriver/Private/DriverSequence.cpp
--- a/Engine/Source/Developer/AutomationDriver/Private/DriverSequence.cpp
+++ b/Engine/Source/Developer/AutomationDriver/Private/DriverSequence.cpp
@@ -728,6 +728,10 @@ public:
 
 	TAsyncResult<bool> Perform()
 	{
+		if (!Application->IsHandlingMessages())
+		{
+			return false;
+		}
 		return StepsExecutor->Execute();
 	}
 
@@ -735,7 +739,7 @@ private:
 
 	FAsyncActionSequence(
 		const TSharedRef<FAsyncAutomationDriver, ESPMode::ThreadSafe>& InAsyncDriver,
-		FAutomatedApplication* const InApplication,
+		const TSharedRef<FAutomatedApplication, ESPMode::ThreadSafe>& InApplication,
 		const TSharedRef<IStepExecutor, ESPMode::ThreadSafe>& InStepsExecutor)
 		: AsyncDriver(InAsyncDriver)
 		, Application(InApplication)
@@ -1338,7 +1342,7 @@ private:
 private:
 
 	const TSharedRef<FAsyncAutomationDriver, ESPMode::ThreadSafe> AsyncDriver;
-	FAutomatedApplication* const Application;
+	const TSharedRef<FAutomatedApplication, ESPMode::ThreadSafe> Application;
 	const TSharedRef<IStepExecutor, ESPMode::ThreadSafe> StepsExecutor;
 
 	friend FAsyncActionSequenceFactory;
@@ -1346,12 +1350,12 @@ private:
 
 TSharedRef<FAsyncActionSequence, ESPMode::ThreadSafe> FAsyncActionSequenceFactory::Create(
 	const TSharedRef<FAsyncAutomationDriver, ESPMode::ThreadSafe>& AsyncDriver,
-	FAutomatedApplication* const Application)
+	const TSharedRef<FAutomatedApplication, ESPMode::ThreadSafe>& Application)
 {
 	return MakeShareable(new FAsyncActionSequence(
 		AsyncDriver,
 		Application,
-		FStepExecutorFactory::Create(AsyncDriver->GetConfiguration())));
+		FStepExecutorFactory::Create(AsyncDriver->GetConfiguration(), Application)));
 }
 
 class FAsyncDriverSequence
@@ -1854,7 +1858,7 @@ private:
 TSharedRef<FActionSequence, ESPMode::ThreadSafe> FActionSequenceFactory::Create(
 	const TSharedRef<FAutomationDriver, ESPMode::ThreadSafe>& Driver,
 	const TSharedRef<FAsyncAutomationDriver, ESPMode::ThreadSafe>& AsyncDriver,
-	FAutomatedApplication* const Application)
+	const TSharedRef<FAutomatedApplication, ESPMode::ThreadSafe>& Application)
 {
 	return MakeShareable(new FActionSequence(
 		Driver,
diff --git a/Engine/Source/Developer/AutomationDriver/Private/DriverSequence.h b/Engine/Source/Developer/AutomationDriver/Private/DriverSequence.h
--- a/Engine/Source/Developer/AutomationDriver/Private/DriverSequence.h
+++ b/Engine/Source/Developer/AutomationDriver/Private/DriverSequence.h
@@ -17,7 +17,7 @@ public:
 
 	static TSharedRef<FAsyncActionSequence, ESPMode::ThreadSafe> Create(
 		const TSharedRef<FAsyncAutomationDriver, ESPMode::ThreadSafe>& AsyncDriver,
-		FAutomatedApplication* const Application);
+		const TSharedRef<FAutomatedApplication, ESPMode::ThreadSafe>& Application);
 };
 
 class FAsyncDriverSequenceFactory
@@ -41,7 +41,7 @@ public:
 	static TSharedRef<FActionSequence, ESPMode::ThreadSafe> Create(
 		const TSharedRef<FAutomationDriver, ESPMode::ThreadSafe>& Driver,
 		const TSharedRef<FAsyncAutomationDriver, ESPMode::ThreadSafe>& AsyncDriver,
-		FAutomatedApplication* const Application);
+		const TSharedRef<FAutomatedApplication, ESPMode::ThreadSafe>& Application);
 };
 
 class FDriverSequenceFactory
diff --git a/Engine/Source/Developer/AutomationDriver/Private/StepExecutor.cpp b/Engine/Source/Developer/AutomationDriver/Private/StepExecutor.cpp
--- a/Engine/Source/Developer/AutomationDriver/Private/StepExecutor.cpp
+++ b/Engine/Source/Developer/AutomationDriver/Private/StepExecutor.cpp
@@ -3,6 +3,7 @@
 #include "StepExecutor.h"
 #include "IStepExecutor.h"
 #include "IAutomationDriver.h"
+#include "AutomatedApplication.h"
 
 #include "DriverConfiguration.h"
 
@@ -86,8 +87,10 @@ public:
 private:
 
 	FStepExecutor(
-		const TSharedRef<FDriverConfiguration, ESPMode::ThreadSafe>& InConfiguration)
+		const TSharedRef<FDriverConfiguration, ESPMode::ThreadSafe>& InConfiguration,
+		const TSharedRef<FAutomatedApplication, ESPMode::ThreadSafe>& InApplication)
 		: Configuration(InConfiguration)
+		, Application(InApplication)
 		, Steps()
 		, CurrentStepIndex(0)
 		, Promise()
@@ -107,7 +110,7 @@ private:
 
 			// If we've encountered an invalid step that's greater then zero then we were just waiting
 			// a little bit after the last step completed before signaling completion.
-			if (StepIndex > 0 && !Steps.IsValidIndex(StepIndex))
+			if ((StepIndex > 0 && !Steps.IsValidIndex(StepIndex)) || !Application->IsHandlingMessages())
 			{
 				Promise->SetValue(true);
 				Promise.Reset();
@@ -153,6 +156,7 @@ private:
 private:
 
 	const TSharedRef<FDriverConfiguration, ESPMode::ThreadSafe> Configuration;
+	const TSharedRef<FAutomatedApplication, ESPMode::ThreadSafe>& Application;
 
 	TArray<FExecuteStepDelegate> Steps;
 	int32 CurrentStepIndex;
@@ -166,7 +170,8 @@ private:
 };
 
 TSharedRef<IStepExecutor, ESPMode::ThreadSafe> FStepExecutorFactory::Create(
-	const TSharedRef<FDriverConfiguration, ESPMode::ThreadSafe>& Configuration)
+	const TSharedRef<FDriverConfiguration, ESPMode::ThreadSafe>& Configuration,
+	const TSharedRef<FAutomatedApplication, ESPMode::ThreadSafe>& Application)
 {
-	return MakeShareable(new FStepExecutor(Configuration));
+	return MakeShareable(new FStepExecutor(Configuration, Application));
 }
diff --git a/Engine/Source/Developer/AutomationDriver/Private/StepExecutor.h b/Engine/Source/Developer/AutomationDriver/Private/StepExecutor.h
--- a/Engine/Source/Developer/AutomationDriver/Private/StepExecutor.h
+++ b/Engine/Source/Developer/AutomationDriver/Private/StepExecutor.h
@@ -5,12 +5,14 @@
 
 class IStepExecutor;
 class FDriverConfiguration;
+class FAutomatedApplication;
 
 class FStepExecutorFactory
 {
 public:
 
 	static TSharedRef<IStepExecutor, ESPMode::ThreadSafe> Create(
-		const TSharedRef<FDriverConfiguration, ESPMode::ThreadSafe>& Configuration);
+		const TSharedRef<FDriverConfiguration, ESPMode::ThreadSafe>& Configuration,
+		const TSharedRef<FAutomatedApplication, ESPMode::ThreadSafe>& Application);
 
 };
diff --git a/Engine/Source/Runtime/Core/Public/Misc/AutomationTest.h b/Engine/Source/Runtime/Core/Public/Misc/AutomationTest.h
--- a/Engine/Source/Runtime/Core/Public/Misc/AutomationTest.h
+++ b/Engine/Source/Runtime/Core/Public/Misc/AutomationTest.h
@@ -2251,6 +2251,7 @@ private:
 					return true;
 				}
 
+				bDone = false;
 				Predicate(FDoneDelegate::CreateSP(this, &FUntilDoneLatentCommand::Done));
 				bIsRunning = true;
 				StartedRunning = FDateTime::UtcNow();
@@ -2275,7 +2276,10 @@ private:
 
 		void Done()
 		{
-			bDone = true;
+			if (bIsRunning)
+			{
+				bDone = true;
+			}
 		}
 
 		void Reset()
@@ -2322,6 +2326,7 @@ private:
 					return true;
 				}
 
+				bDone = false;
 				Future = Async(Execution, [this]() {
 					Predicate(FDoneDelegate::CreateRaw(this, &FAsyncUntilDoneLatentCommand::Done));
 				});
@@ -2348,14 +2353,17 @@ private:
 
 		void Done()
 		{
-			bDone = true;
+			if (Future.IsValid())
+			{
+				bDone = true;
+			}
 		}
 
 		void Reset()
 		{
 			// Reset the done for the next potential run of this command
 			bDone = false;
-			Future = TFuture<void>();
+			Future.Reset();
 		}
 
 	private:
@@ -2396,9 +2404,10 @@ private:
 					return true;
 				}
 
+				bDone = false;
 				Future = Async(Execution, [this]() {
 					Predicate();
-					bDone = true;
+					Done();
 				});
 
 				StartedRunning = FDateTime::UtcNow();
@@ -2423,14 +2432,17 @@ private:
 
 		void Done()
 		{
-			bDone = true;
+			if (Future.IsValid())
+			{
+				bDone = true;
+			}
 		}
 
 		void Reset()
 		{
 			// Reset the done for the next potential run of this command
 			bDone = false;
-			Future = TFuture<void>();
+			Future.Reset();
 		}
 
 	private:
