Fix DoubleClick self tests for Automation Driver

diff --git a/Engine/Source/Developer/AutomationDriver/Private/IStepExecutor.h b/Engine/Source/Developer/AutomationDriver/Private/IStepExecutor.h
--- a/Engine/Source/Developer/AutomationDriver/Private/IStepExecutor.h
+++ b/Engine/Source/Developer/AutomationDriver/Private/IStepExecutor.h
@@ -33,13 +33,13 @@ class IStepExecutor
 public:
 
 	DECLARE_DELEGATE_RetVal_OneParam(FStepResult, FExecuteStepDelegate, const FTimespan& /*TotalProcessTime*/);
-	virtual void Add(const FExecuteStepDelegate& Step) = 0;
+	virtual void Add(const TSharedRef<FExecuteStepDelegate>& Step) = 0;
 
-	virtual void Add(const TFunction<FStepResult(const FTimespan&)>& Step) = 0;
+	virtual void Add(const TFunction<FStepResult(const FTimespan&)>& StepFunction) = 0;
 
-	virtual void InsertNext(const FExecuteStepDelegate& Step) = 0;
+	virtual void InsertNext(const TSharedRef<FExecuteStepDelegate>& Step) = 0;
 
-	virtual void InsertNext(const TFunction<FStepResult(const FTimespan&)>& Step) = 0;
+	virtual void InsertNext(const TFunction<FStepResult(const FTimespan&)>& StepFunction) = 0;
 
 	virtual TAsyncResult<bool> Execute() = 0;
 
diff --git a/Engine/Source/Developer/AutomationDriver/Private/StepExecutor.cpp b/Engine/Source/Developer/AutomationDriver/Private/StepExecutor.cpp
--- a/Engine/Source/Developer/AutomationDriver/Private/StepExecutor.cpp
+++ b/Engine/Source/Developer/AutomationDriver/Private/StepExecutor.cpp
@@ -24,34 +24,40 @@ public:
 		{
 			Promise->SetValue(false);
 		}
+
+		// As we use guards to access to the Steps array we have to clear it explicitly to sync access to it
+		FScopeLock StateLock(&StepsCS);
+		Steps.Empty();
 	}
 
-	virtual void Add(const FExecuteStepDelegate& Step) override
+	virtual void Add(const TSharedRef<FExecuteStepDelegate>& Step) override
 	{
 		check(!Promise.IsValid());
 		FScopeLock StateLock(&StepsCS);
 		Steps.Add(Step);
 	}
 
-	virtual void Add(const TFunction<FStepResult(const FTimespan&)>& Step) override
+	virtual void Add(const TFunction<FStepResult(const FTimespan&)>& StepFunction) override
 	{
-		check(!Promise.IsValid());
-		FScopeLock StateLock(&StepsCS);
-		Steps.Add(FExecuteStepDelegate::CreateLambda(Step));
+		check(StepFunction);
+		TSharedRef<FExecuteStepDelegate> Step = MakeShared<FExecuteStepDelegate>(
+			FExecuteStepDelegate::CreateLambda(StepFunction));
+		Add(Step);
 	}
 
-	virtual void InsertNext(const FExecuteStepDelegate& Step) override
+	virtual void InsertNext(const TSharedRef<FExecuteStepDelegate>& Step) override
 	{
 		check(Promise.IsValid());
 		FScopeLock StateLock(&StepsCS);
 		Steps.Insert(Step, CurrentStepIndex + 1);
 	}
 
-	virtual void InsertNext(const TFunction<FStepResult(const FTimespan&)>& Step) override
+	virtual void InsertNext(const TFunction<FStepResult(const FTimespan&)>& StepFunction) override
 	{
-		check(Promise.IsValid());
-		FScopeLock StateLock(&StepsCS);
-		Steps.Insert(FExecuteStepDelegate::CreateLambda(Step), CurrentStepIndex + 1);
+		check(StepFunction);
+		TSharedRef<FExecuteStepDelegate> Step = MakeShared<FExecuteStepDelegate>(
+			FExecuteStepDelegate::CreateLambda(StepFunction));
+		InsertNext(Step);
 	}
 
 	virtual TAsyncResult<bool> Execute() override
@@ -120,7 +126,7 @@ private:
 
 			check(Steps.IsValidIndex(StepIndex));
 
-			Result = Steps[StepIndex].Execute(StepTotalProcessTime);
+			Result = Steps[StepIndex]->Execute(StepTotalProcessTime);
 		}
 
 		if (Result.State == FStepResult::EState::FAILED)
@@ -158,7 +164,7 @@ private:
 	const TSharedRef<FDriverConfiguration, ESPMode::ThreadSafe> Configuration;
 	const TSharedRef<FAutomatedApplication, ESPMode::ThreadSafe>& Application;
 
-	TArray<FExecuteStepDelegate> Steps;
+	TArray<TSharedRef<FExecuteStepDelegate>> Steps;
 	int32 CurrentStepIndex;
 	TSharedPtr<TPromise<bool>> Promise;
 	FTimespan StepTotalProcessTime;
