Fix DoubleClick, Scroll and Press tests of Automation Driver

diff --git a/Engine/Source/Developer/AutomationDriver/Private/AutomatedApplication.cpp b/Engine/Source/Developer/AutomationDriver/Private/AutomatedApplication.cpp
--- a/Engine/Source/Developer/AutomationDriver/Private/AutomatedApplication.cpp
+++ b/Engine/Source/Developer/AutomationDriver/Private/AutomatedApplication.cpp
@@ -19,7 +19,10 @@ public:
 		{
 			if (RealCursor.IsValid())
 			{
-				FakePosition = RealCursor->GetPosition();
+				if (!bOverrideRealCursor)
+				{
+					FakePosition = RealCursor->GetPosition();
+				}
 			}
 			else
 			{
@@ -132,9 +135,15 @@ public:
 		bAllowMessageHandling = bValue;
 	}
 
+	void SetOverrideRealCursorCoordinates(bool bOverride)
+	{
+		bOverrideRealCursor = bOverride;
+	}
+
 	FAutomatedCursor(const TSharedPtr<ICursor>& InRealCursor)
 		: RealCursor(InRealCursor)
 		, bAllowMessageHandling(false)
+		, bOverrideRealCursor(false)
 	{
 		if (RealCursor.IsValid())
 		{
@@ -156,6 +165,7 @@ private:
 	const TSharedPtr<ICursor> RealCursor;
 
 	bool bAllowMessageHandling;
+	bool bOverrideRealCursor;
 
 	mutable FVector2D FakePosition;
 	mutable EMouseCursor::Type FakeMouseType;
@@ -201,7 +211,7 @@ public:
 		}
 	}
 
-	virtual bool IsHandlingMessages() override
+	virtual bool IsHandlingMessages() const override
 	{
 		if (!PassThroughMessageHandler.IsValid() || !PassThroughMessageHandler->IsHandlingMessages())
 		{
@@ -216,6 +226,14 @@ public:
 		return true;
 	}
 
+	virtual void SetOverrideRealCursorCoordinates(bool bOverride) override
+	{
+		if (AutomatedCursor.IsValid())
+		{
+			AutomatedCursor->SetOverrideRealCursorCoordinates(bOverride);
+		}
+	}
+
 	virtual void SetMessageHandler(const TSharedRef<FGenericApplicationMessageHandler>& InMessageHandler) override
 	{
 		RealMessageHandler = InMessageHandler;
@@ -277,7 +295,7 @@ public:
 
 	virtual FModifierKeysState GetModifierKeys() const override
 	{
-		if (!PassThroughMessageHandler.IsValid() || !PassThroughMessageHandler->IsHandlingMessages())
+		if (IsHandlingMessages())
 		{
 			return FakeModifierKeys;
 		}
diff --git a/Engine/Source/Developer/AutomationDriver/Private/AutomatedApplication.h b/Engine/Source/Developer/AutomationDriver/Private/AutomatedApplication.h
--- a/Engine/Source/Developer/AutomationDriver/Private/AutomatedApplication.h
+++ b/Engine/Source/Developer/AutomationDriver/Private/AutomatedApplication.h
@@ -18,7 +18,8 @@ public:
 
 	virtual void AllowPlatformMessageHandling() = 0;
 	virtual void DisablePlatformMessageHandling() = 0;
-	virtual bool IsHandlingMessages() = 0;
+	virtual bool IsHandlingMessages() const = 0;
+	virtual void SetOverrideRealCursorCoordinates(bool bOverrideRealCursor) = 0;
 
 	virtual void SetFakeModifierKeys(FModifierKeysState Value) = 0;
 
diff --git a/Engine/Source/Developer/AutomationDriver/Private/AutomationDriverModule.cpp b/Engine/Source/Developer/AutomationDriver/Private/AutomationDriverModule.cpp
--- a/Engine/Source/Developer/AutomationDriver/Private/AutomationDriverModule.cpp
+++ b/Engine/Source/Developer/AutomationDriver/Private/AutomationDriverModule.cpp
@@ -86,6 +86,11 @@ public:
 		RealApplication->SetMessageHandler(RealMessageHandler.ToSharedRef());
 
 		AutomatedApplication.Reset();
+		// Note that here we have to return the platform cursor that was overriden by FAutomatedCursor while enabling the module.
+		// Sometimes it is not being replaced with the real platform cursor, so we force it.
+		// It is needed to prevent possible issues that we are sending mouse events to the cursor that is outdated while executing the next tests from the list of AutomationDriver tests.
+		FSlateApplication::Get().UsePlatformCursorForCursorUser(true);
+
 		RealApplication.Reset();
 		RealMessageHandler.Reset();
 	}
diff --git a/Engine/Source/Developer/AutomationDriver/Private/StepExecutor.cpp b/Engine/Source/Developer/AutomationDriver/Private/StepExecutor.cpp
--- a/Engine/Source/Developer/AutomationDriver/Private/StepExecutor.cpp
+++ b/Engine/Source/Developer/AutomationDriver/Private/StepExecutor.cpp
@@ -114,10 +114,16 @@ private:
 		{
 			FScopeLock StateLock(&StepsCS);
 
+			if (0 == StepIndex)
+			{
+				Application->SetOverrideRealCursorCoordinates(true);
+			}
+
 			// If we've encountered an invalid step that's greater then zero then we were just waiting
 			// a little bit after the last step completed before signaling completion.
 			if ((StepIndex > 0 && !Steps.IsValidIndex(StepIndex)) || !Application->IsHandlingMessages())
 			{
+				Application->SetOverrideRealCursorCoordinates(false); 
 				Promise->SetValue(true);
 				Promise.Reset();
 				StepTotalProcessTime = FTimespan::Zero();
@@ -131,6 +137,7 @@ private:
 
 		if (Result.State == FStepResult::EState::FAILED)
 		{
+			Application->SetOverrideRealCursorCoordinates(false);
 			Promise->SetValue(false);
 			Promise.Reset();
 			StepTotalProcessTime = FTimespan::Zero();
