Unify internal logic of FAutomatedApplicationImpl that works with real and fake inner instances

diff --git a/Engine/Source/Developer/AutomationDriver/Private/AutomatedApplication.cpp b/Engine/Source/Developer/AutomationDriver/Private/AutomatedApplication.cpp
--- a/Engine/Source/Developer/AutomationDriver/Private/AutomatedApplication.cpp
+++ b/Engine/Source/Developer/AutomationDriver/Private/AutomatedApplication.cpp
@@ -275,17 +275,17 @@ public:
 
 	virtual void SetCapture(const TSharedPtr<FGenericWindow>& Window) override
 	{
-		if (!PassThroughMessageHandler.IsValid() || PassThroughMessageHandler->IsHandlingMessages())
+		if (IsHandlingMessages())
 		{
-			RealApplication->SetCapture(Window);
+			FakeCapture = Window;
 		}
 
-		FakeCapture = Window;
+		RealApplication->SetCapture(Window);
 	}
 
 	virtual void* GetCapture(void) const override
 	{
-		if (!PassThroughMessageHandler.IsValid() || !PassThroughMessageHandler->IsHandlingMessages())
+		if (IsHandlingMessages())
 		{
 			return (void*)FakeCapture.Get();
 		}
@@ -329,7 +329,7 @@ public:
 
 	virtual bool IsCursorDirectlyOverSlateWindow() const override
 	{
-		if (PassThroughMessageHandler.IsValid() && !PassThroughMessageHandler->IsHandlingMessages())
+		if (IsHandlingMessages())
 		{
 			return InternalGetWindowUnderCursor().IsValid();
 		}
@@ -339,7 +339,7 @@ public:
 
 	virtual TSharedPtr<FGenericWindow> GetWindowUnderCursor() override
 	{
-		if (PassThroughMessageHandler.IsValid() && !PassThroughMessageHandler->IsHandlingMessages())
+		if (IsHandlingMessages())
 		{
 			TSharedPtr<SWindow> Window = InternalGetWindowUnderCursor();
 			if (Window.IsValid())
