AutomationDriver - Fix tests' timing issues

diff --git a/Engine/Plugins/Tests/AutomationDriverTests/Source/AutomationDriverTests/Private/AutomationDriver.spec.cpp b/Engine/Plugins/Tests/AutomationDriverTests/Source/AutomationDriverTests/Private/AutomationDriver.spec.cpp
--- a/Engine/Plugins/Tests/AutomationDriverTests/Source/AutomationDriverTests/Private/AutomationDriver.spec.cpp
+++ b/Engine/Plugins/Tests/AutomationDriverTests/Source/AutomationDriverTests/Private/AutomationDriver.spec.cpp
@@ -925,6 +925,7 @@ void FAutomationDriverSpec::Define()
 			It("should return false if the element is not currently visible in the SWidget DOM being displayed", EAsyncExecution::ThreadPool, [this]()
 			{
 				SuiteViewModel->SetPianoVisibility(EVisibility::Hidden);
+				Driver->Wait(Until::ElementIsHidden(By::Id("Piano"), FWaitTimeout::InSeconds(1)));
 				TEST_FALSE(Driver->FindElement(By::Id("Piano"))->IsVisible());
 			});
 
@@ -938,6 +939,7 @@ void FAutomationDriverSpec::Define()
 		{
 			It("should return true if the element is a currently enabled SWidget", EAsyncExecution::ThreadPool, [this]()
 			{
+				Driver->Wait(Until::ElementIsVisible(By::Id("Piano"), FWaitTimeout::InSeconds(1)));
 				TEST_TRUE(Driver->FindElement(By::Id("KeyB"))->IsInteractable());
 			});
 
diff --git a/Engine/Source/Developer/AutomationDriver/Private/WaitUntil.cpp b/Engine/Source/Developer/AutomationDriver/Private/WaitUntil.cpp
--- a/Engine/Source/Developer/AutomationDriver/Private/WaitUntil.cpp
+++ b/Engine/Source/Developer/AutomationDriver/Private/WaitUntil.cpp
@@ -145,6 +145,46 @@ FDriverWaitDelegate Until::ElementIsVisible(const TSharedRef<IElementLocator, ES
 	});
 }
 
+FDriverWaitDelegate Until::ElementIsHidden(const TSharedRef<IElementLocator, ESPMode::ThreadSafe>& ElementLocator, FWaitTimeout Timeout)
+{
+	return Until::ElementIsHidden(ElementLocator, FWaitInterval::InSeconds(1.0), Timeout);
+}
+
+FDriverWaitDelegate Until::ElementIsHidden(const TSharedRef<IElementLocator, ESPMode::ThreadSafe>& ElementLocator, FWaitInterval Interval, FWaitTimeout Timeout)
+{
+	return Until::Lambda([ElementLocator, Interval, Timeout](const FTimespan& TotalWaitTime) {
+		check(IsInGameThread());
+
+		TArray<TSharedRef<IApplicationElement>> Elements;
+		ElementLocator->Locate(Elements);
+
+		if (Elements.Num() == 0)
+		{
+			if (TotalWaitTime > Timeout.Timespan)
+			{
+				return FDriverWaitResponse::Failed();
+			}
+
+			return FDriverWaitResponse::Wait(Interval.Timespan);
+		}
+
+		for (TSharedRef<IApplicationElement> Element : Elements)
+		{
+			if (Element->IsVisible())
+			{
+				if (TotalWaitTime > Timeout.Timespan)
+				{
+					return FDriverWaitResponse::Failed();
+				}
+
+				return FDriverWaitResponse::Wait(Interval.Timespan);
+			}
+		}
+
+		return FDriverWaitResponse::Passed();
+		});
+}
+
 FDriverWaitDelegate Until::ElementIsInteractable(const TSharedRef<IElementLocator, ESPMode::ThreadSafe>& ElementLocator, FWaitTimeout Timeout)
 {
 	return Until::ElementIsInteractable(ElementLocator, FWaitInterval::InSeconds(1.0), Timeout);
diff --git a/Engine/Source/Developer/AutomationDriver/Public/WaitUntil.h b/Engine/Source/Developer/AutomationDriver/Public/WaitUntil.h
--- a/Engine/Source/Developer/AutomationDriver/Public/WaitUntil.h
+++ b/Engine/Source/Developer/AutomationDriver/Public/WaitUntil.h
@@ -119,6 +119,17 @@ public:
 	 */
 	static FDriverWaitDelegate ElementIsVisible(const TSharedRef<IElementLocator, ESPMode::ThreadSafe>& ElementLocator, FWaitInterval Interval, FWaitTimeout Timeout);
 
+	/**
+	* Creates a new wait delegate which completes it's wait only if the specified element locator discovers hidden elements or if the specified timeout timespan elapses
+	*/
+	static FDriverWaitDelegate ElementIsHidden(const TSharedRef<IElementLocator, ESPMode::ThreadSafe>& ElementLocator, FWaitTimeout Timeout);
+
+	/**
+	 * Creates a new wait delegate which completes it's wait only if the specified element locator discovers hidden elements or if the specified timeout timespan elapses;
+	 * The element locator is only re-evaluated at the specified wait interval
+	 */
+	static FDriverWaitDelegate ElementIsHidden(const TSharedRef<IElementLocator, ESPMode::ThreadSafe>& ElementLocator, FWaitInterval Interval, FWaitTimeout Timeout);
+
 	/**
 	 * Creates a new wait delegate which completes it's wait only if the specified element locator discovers interactable elements or if the specified timeout timespan elapses
 	 */
