Fix the early destruction of FAsyncActionSequence in FAsyncAutomationDriver

The functions within the FAsyncAutomationDriver that implicitly rely on FAsyncActionSequence were not functioning correctly due to the action sequence object being destroyed right after leaving the scope of the IAsyncAutomationDriver interface function.
To ensure that the sequence object remains valid until the actions are performed, a shared pointer to this object is stored within the FAsyncAutomationDriver.

#rnx

diff --git a/Engine/Source/Developer/AutomationDriver/Private/AutomationDriver.cpp b/Engine/Source/Developer/AutomationDriver/Private/AutomationDriver.cpp
--- a/Engine/Source/Developer/AutomationDriver/Private/AutomationDriver.cpp
+++ b/Engine/Source/Developer/AutomationDriver/Private/AutomationDriver.cpp
@@ -1,6 +1,7 @@
 // Copyright Epic Games, Inc. All Rights Reserved.
 
 #include "AutomationDriver.h"
+#include "AutomationDriverLogging.h"
 #include "IApplicationElement.h"
 #include "IDriverSequence.h"
 
@@ -18,6 +19,36 @@ FAsyncAutomationDriver::~FAsyncAutomationDriver()
 	Application->SetFakeModifierKeys(FModifierKeysState());
 }
 
+bool FAsyncAutomationDriver::PinActionSequence(const TSharedPtr<IAsyncActionSequence, ESPMode::ThreadSafe>& Sequence)
+{
+	if (PinnedSequence && PinnedSequence->IsExecuting())
+	{
+		FAutomationDriverLogging::CannotExecuteMultipleActionSequencesInParallel();
+		return false;
+	}
+
+	PinnedSequence = Sequence;
+	return true;
+}
+
+bool FAsyncAutomationDriver::UnpinActionSequence(const TSharedPtr<IAsyncActionSequence, ESPMode::ThreadSafe>& Sequence)
+{
+	if (PinnedSequence != Sequence)
+	{
+		FAutomationDriverLogging::CannotUnpinActionSequenceIfNotPinned();
+		return false;
+	}
+
+	if (PinnedSequence && PinnedSequence->IsExecuting())
+	{
+		FAutomationDriverLogging::CannotUnpinActionSequenceIfExecuting();
+		return false;
+	}
+
+	PinnedSequence.Reset();
+	return true;
+}
+
 TAsyncResult<bool> FAsyncAutomationDriver::Wait(FTimespan Timespan)
 {
 	TSharedRef<IAsyncDriverSequence, ESPMode::ThreadSafe> Sequence = CreateSequence();
diff --git a/Engine/Source/Developer/AutomationDriver/Private/AutomationDriver.h b/Engine/Source/Developer/AutomationDriver/Private/AutomationDriver.h
--- a/Engine/Source/Developer/AutomationDriver/Private/AutomationDriver.h
+++ b/Engine/Source/Developer/AutomationDriver/Private/AutomationDriver.h
@@ -9,6 +9,7 @@
 class FAutomatedApplication;
 class IElementLocator;
 class IAsyncDriverSequence;
+class IAsyncActionSequence;
 class IAsyncDriverElement;
 class IAsyncDriverElementCollection;
 class FDriverConfiguration;
@@ -39,6 +40,20 @@ public:
 
 	virtual TSharedRef<FDriverConfiguration, ESPMode::ThreadSafe> GetConfiguration() const override;
 
+	/**
+	 * Saves a shared pointer to the action sequence object to ensure it remains in memory until execution is complete.
+	 * @param Sequence shared pointer to the action sequence.
+	 * @return true if the pointer has been saved successfully; false if another pinned sequence is currently being executed.
+	 */
+	bool PinActionSequence(const TSharedPtr<IAsyncActionSequence, ESPMode::ThreadSafe>& Sequence);
+	
+	/**
+	 * Resets an internally stored shared pointer to the previously pinned action sequence object so that it can be deleted.
+	 * @param Sequence shared pointer to the action sequence.
+	 * @return true if the pointer has been reset successfully; false if the provided sequence hasn't been pinned previously or if it is still executing
+	 */
+	bool UnpinActionSequence(const TSharedPtr<IAsyncActionSequence, ESPMode::ThreadSafe>& Sequence);
+
 public:
 
 	void TrackPress(int32 KeyCode, int32 CharCode);
@@ -67,6 +82,7 @@ private:
 
 	const TSharedRef<FAutomatedApplication, ESPMode::ThreadSafe> Application;
 	const TSharedRef<FDriverConfiguration, ESPMode::ThreadSafe> Configuration;
+	TSharedPtr<IAsyncActionSequence, ESPMode::ThreadSafe> PinnedSequence;
 
 	TSet<FKey> PressedModifiers;
 	TSet<int32> PressedKeys;
diff --git a/Engine/Source/Developer/AutomationDriver/Private/AutomationDriverLogging.cpp b/Engine/Source/Developer/AutomationDriver/Private/AutomationDriverLogging.cpp
--- a/Engine/Source/Developer/AutomationDriver/Private/AutomationDriverLogging.cpp
+++ b/Engine/Source/Developer/AutomationDriver/Private/AutomationDriverLogging.cpp
@@ -83,3 +83,18 @@ void FAutomationDriverLogging::CannotClickUnhoveredElement(const TSharedPtr<IEle
 	}
 }
 
+void FAutomationDriverLogging::CannotExecuteMultipleActionSequencesInParallel()
+{
+	UE_LOG(LogAutomationDriver, Error, TEXT("Parallel execution of multiple action sequences is not supported"));
+}
+
+void FAutomationDriverLogging::CannotUnpinActionSequenceIfNotPinned()
+{
+	UE_LOG(LogAutomationDriver, Error, TEXT("Cannot unpin action sequence if it's not pinned"));
+}
+
+void FAutomationDriverLogging::CannotUnpinActionSequenceIfExecuting()
+{
+	UE_LOG(LogAutomationDriver, Error, TEXT("Cannot unpin action sequence if it's still executing"));
+}
+
diff --git a/Engine/Source/Developer/AutomationDriver/Private/AutomationDriverLogging.h b/Engine/Source/Developer/AutomationDriver/Private/AutomationDriverLogging.h
--- a/Engine/Source/Developer/AutomationDriver/Private/AutomationDriverLogging.h
+++ b/Engine/Source/Developer/AutomationDriver/Private/AutomationDriverLogging.h
@@ -23,4 +23,10 @@ public:
 	static void ElementHasNoWindow(const TSharedPtr<IElementLocator, ESPMode::ThreadSafe>& ElementLocator);
 
 	static void CannotClickUnhoveredElement(const TSharedPtr<IElementLocator, ESPMode::ThreadSafe>& ElementLocator);
+	
+	static void CannotExecuteMultipleActionSequencesInParallel();
+
+	static void CannotUnpinActionSequenceIfNotPinned();
+
+	static void CannotUnpinActionSequenceIfExecuting();
 };
diff --git a/Engine/Source/Developer/AutomationDriver/Private/DriverSequence.cpp b/Engine/Source/Developer/AutomationDriver/Private/DriverSequence.cpp
--- a/Engine/Source/Developer/AutomationDriver/Private/DriverSequence.cpp
+++ b/Engine/Source/Developer/AutomationDriver/Private/DriverSequence.cpp
@@ -191,6 +191,11 @@ public:
 	virtual ~FAsyncActionSequence()
 	{ }
 
+	bool IsExecuting() override
+	{
+		return StepsExecutor->IsExecuting();
+	}
+
 	virtual IAsyncActionSequence& Wait(FTimespan Timespan) override
 	{
 		InternalWait(Until::Lambda([Timespan](const FTimespan& TotalWaitTime) {
@@ -732,6 +737,21 @@ public:
 		{
 			return false;
 		}
+
+		if (!AsyncDriver->PinActionSequence(SharedThis(this)))
+		{
+			return false;
+		}
+
+		TWeakPtr<FAsyncActionSequence, ESPMode::ThreadSafe> LocalWeakThis(SharedThis(this));
+		StepsExecutor->OnCompleted([LocalWeakThis]()
+		{
+			if (TSharedPtr<FAsyncActionSequence, ESPMode::ThreadSafe> Sequence = LocalWeakThis.Pin())
+			{
+				Sequence->AsyncDriver->UnpinActionSequence(Sequence);
+			}
+		});
+
 		return StepsExecutor->Execute();
 	}
 
diff --git a/Engine/Source/Developer/AutomationDriver/Private/IStepExecutor.h b/Engine/Source/Developer/AutomationDriver/Private/IStepExecutor.h
--- a/Engine/Source/Developer/AutomationDriver/Private/IStepExecutor.h
+++ b/Engine/Source/Developer/AutomationDriver/Private/IStepExecutor.h
@@ -44,4 +44,6 @@ public:
 	virtual TAsyncResult<bool> Execute() = 0;
 
 	virtual bool IsExecuting() const = 0;
+
+	virtual void OnCompleted(TFunction<void()> Callback) = 0;
 };
diff --git a/Engine/Source/Developer/AutomationDriver/Private/StepExecutor.cpp b/Engine/Source/Developer/AutomationDriver/Private/StepExecutor.cpp
--- a/Engine/Source/Developer/AutomationDriver/Private/StepExecutor.cpp
+++ b/Engine/Source/Developer/AutomationDriver/Private/StepExecutor.cpp
@@ -89,6 +89,11 @@ public:
 	{
 		return Promise.IsValid();
 	}
+	
+	virtual void OnCompleted(TFunction<void()> Callback) override
+	{
+		OnCompletedCallback = MoveTemp(Callback);
+	}
 
 private:
 
@@ -127,6 +132,7 @@ private:
 				Promise->SetValue(true);
 				Promise.Reset();
 				StepTotalProcessTime = FTimespan::Zero();
+				if (OnCompletedCallback) { OnCompletedCallback(); }
 				return false;
 			}
 
@@ -141,6 +147,7 @@ private:
 			Promise->SetValue(false);
 			Promise.Reset();
 			StepTotalProcessTime = FTimespan::Zero();
+			if (OnCompletedCallback) { OnCompletedCallback(); }
 			return false;
 		}
 
@@ -176,6 +183,7 @@ private:
 	TSharedPtr<TPromise<bool>> Promise;
 	FTimespan StepTotalProcessTime;
 	float LastDelay;
+	TFunction<void()> OnCompletedCallback;
 
 	mutable FCriticalSection StepsCS;
 
diff --git a/Engine/Source/Developer/AutomationDriver/Public/IDriverSequence.h b/Engine/Source/Developer/AutomationDriver/Public/IDriverSequence.h
--- a/Engine/Source/Developer/AutomationDriver/Public/IDriverSequence.h
+++ b/Engine/Source/Developer/AutomationDriver/Public/IDriverSequence.h
@@ -426,6 +426,11 @@ public:
 	 * Queues a change of focus to the element for the specified user
 	 */
 	virtual IAsyncActionSequence& Focus(const TSharedRef<IElementLocator, ESPMode::ThreadSafe>& ElementLocator, uint32 UserFocus) = 0;
+	
+	/**
+	 * Checks whether the sequence is currently executing actions
+	 */
+	virtual bool IsExecuting() = 0;
 };
 
 /**
