Automation Driver: Add support for tagging delegate locators by a debug string, so that failing ones can be identified

diff --git a/Engine/Source/Developer/AutomationDriver/Private/Locators/SlateWidgetLocatorByDelegate.cpp b/Engine/Source/Developer/AutomationDriver/Private/Locators/SlateWidgetLocatorByDelegate.cpp
--- a/Engine/Source/Developer/AutomationDriver/Private/Locators/SlateWidgetLocatorByDelegate.cpp
+++ b/Engine/Source/Developer/AutomationDriver/Private/Locators/SlateWidgetLocatorByDelegate.cpp
@@ -15,11 +15,17 @@ public:
 
 	virtual FString ToDebugString() const
 	{
+		static const FString TypeString = TEXT("[By::Delegate] ");
+		if (!DebugName.IsEmpty())
+		{
+			return TypeString + DebugName;
+		}
+
 		FString DelegateName;
 #if USE_DELEGATE_TRYGETBOUNDFUNCTIONNAME
 		DelegateName = Delegate.TryGetBoundFunctionName().ToString();
 #endif
-		return TEXT("[By::Delegate] ") + DelegateName;
+		return TypeString + DelegateName;
 	}
 
 	virtual void Locate(TArray<TSharedRef<IApplicationElement>>& OutElements) const override
@@ -43,22 +49,23 @@ public:
 private:
 
 	FSlateWidgetLocatorByWidgetDelegate(
-		const FLocateSlateWidgetElementDelegate& InDelegate)
-		: Delegate(InDelegate)
+		const FLocateSlateWidgetElementDelegate& InDelegate, FStringView InDebugName = { })
+		: Delegate(InDelegate), DebugName(InDebugName)
 	{ }
 
 private:
 
 	FLocateSlateWidgetElementDelegate Delegate;
+	FString DebugName;
 
 	friend FSlateWidgetLocatorByDelegateFactory;
 };
 
 
 TSharedRef<IElementLocator, ESPMode::ThreadSafe> FSlateWidgetLocatorByDelegateFactory::Create(
-	const FLocateSlateWidgetElementDelegate& Delegate)
+	const FLocateSlateWidgetElementDelegate& Delegate, FStringView DebugName)
 {
-	return MakeShareable(new FSlateWidgetLocatorByWidgetDelegate(Delegate));
+	return MakeShareable(new FSlateWidgetLocatorByWidgetDelegate(Delegate, DebugName));
 }
 
 class FSlateWidgetLocatorByWidgetPathDelegate
@@ -98,20 +105,21 @@ public:
 private:
 
 	FSlateWidgetLocatorByWidgetPathDelegate(
-		const FLocateSlateWidgetPathElementDelegate& InDelegate)
-		: Delegate(InDelegate)
+		const FLocateSlateWidgetPathElementDelegate& InDelegate, FStringView InDebugName = { })
+		: Delegate(InDelegate), DebugName(InDebugName)
 	{ }
 
 private:
 
 	FLocateSlateWidgetPathElementDelegate Delegate;
+	FString DebugName;
 
 	friend FSlateWidgetLocatorByDelegateFactory;
 };
 
 
 TSharedRef<IElementLocator, ESPMode::ThreadSafe> FSlateWidgetLocatorByDelegateFactory::Create(
-	const FLocateSlateWidgetPathElementDelegate& Delegate)
+	const FLocateSlateWidgetPathElementDelegate& Delegate, FStringView DebugName)
 {
-	return MakeShareable(new FSlateWidgetLocatorByWidgetPathDelegate(Delegate));
+	return MakeShareable(new FSlateWidgetLocatorByWidgetPathDelegate(Delegate, DebugName));
 }
diff --git a/Engine/Source/Developer/AutomationDriver/Private/Locators/SlateWidgetLocatorByDelegate.h b/Engine/Source/Developer/AutomationDriver/Private/Locators/SlateWidgetLocatorByDelegate.h
--- a/Engine/Source/Developer/AutomationDriver/Private/Locators/SlateWidgetLocatorByDelegate.h
+++ b/Engine/Source/Developer/AutomationDriver/Private/Locators/SlateWidgetLocatorByDelegate.h
@@ -10,8 +10,8 @@ class FSlateWidgetLocatorByDelegateFactory
 public:
 
 	static TSharedRef<IElementLocator, ESPMode::ThreadSafe> Create(
-		const FLocateSlateWidgetElementDelegate& Delegate);
+		const FLocateSlateWidgetElementDelegate& Delegate, FStringView DebugName = { });
 
 	static TSharedRef<IElementLocator, ESPMode::ThreadSafe> Create(
-		const FLocateSlateWidgetPathElementDelegate& Delegate);
+		const FLocateSlateWidgetPathElementDelegate& Delegate, FStringView DebugName = { });
 };
