Automation Driver: Add locators and waiting actions for user focus

diff --git a/Engine/Source/Developer/AutomationDriver/Private/LocateBy.cpp b/Engine/Source/Developer/AutomationDriver/Private/LocateBy.cpp
--- a/Engine/Source/Developer/AutomationDriver/Private/LocateBy.cpp
+++ b/Engine/Source/Developer/AutomationDriver/Private/LocateBy.cpp
@@ -8,24 +8,24 @@
 #include "Framework/Application/SlateApplication.h"
 #include "IApplicationElement.h"
 
-TSharedRef<IElementLocator, ESPMode::ThreadSafe> By::Delegate(const FLocateSlateWidgetElementDelegate& Value)
+TSharedRef<IElementLocator, ESPMode::ThreadSafe> By::Delegate(const FLocateSlateWidgetElementDelegate& Value, FStringView DebugName)
 {
-	return FSlateWidgetLocatorByDelegateFactory::Create(Value);
+	return FSlateWidgetLocatorByDelegateFactory::Create(Value, DebugName);
 }
 
-TSharedRef<IElementLocator, ESPMode::ThreadSafe> By::Delegate(const FLocateSlateWidgetPathElementDelegate& Value)
+TSharedRef<IElementLocator, ESPMode::ThreadSafe> By::Delegate(const FLocateSlateWidgetPathElementDelegate& Value, FStringView DebugName)
 {
-	return FSlateWidgetLocatorByDelegateFactory::Create(Value);
+	return FSlateWidgetLocatorByDelegateFactory::Create(Value, DebugName);
 }
 
-TSharedRef<IElementLocator, ESPMode::ThreadSafe> By::WidgetLambda(const TFunction<void(TArray<TSharedRef<SWidget>>&)>& Value)
+TSharedRef<IElementLocator, ESPMode::ThreadSafe> By::WidgetLambda(const TFunction<void(TArray<TSharedRef<SWidget>>&)>& Value, FStringView DebugName)
 {
-	return FSlateWidgetLocatorByDelegateFactory::Create(FLocateSlateWidgetElementDelegate::CreateLambda(Value));
+	return FSlateWidgetLocatorByDelegateFactory::Create(FLocateSlateWidgetElementDelegate::CreateLambda(Value), DebugName);
 }
 
-TSharedRef<IElementLocator, ESPMode::ThreadSafe> By::WidgetPathLambda(const TFunction<void(TArray<FWidgetPath>&)>& Value)
+TSharedRef<IElementLocator, ESPMode::ThreadSafe> By::WidgetPathLambda(const TFunction<void(TArray<FWidgetPath>&)>& Value, FStringView DebugName)
 {
-	return FSlateWidgetLocatorByDelegateFactory::Create(FLocateSlateWidgetPathElementDelegate::CreateLambda(Value));
+	return FSlateWidgetLocatorByDelegateFactory::Create(FLocateSlateWidgetPathElementDelegate::CreateLambda(Value), DebugName);
 }
 
 TSharedRef<IElementLocator, ESPMode::ThreadSafe> By::Id(const FString& Value)
@@ -122,6 +122,26 @@ TSharedRef<IElementLocator, ESPMode::ThreadSafe> By::Cursor()
 	});
 }
 
+TSharedRef<IElementLocator, ESPMode::ThreadSafe> By::KeyboardFocus()
+{
+	return WidgetLambda([](TArray<TSharedRef<SWidget>>& OutWidgets) -> void {
+		if (const TSharedPtr<SWidget> Widget = FSlateApplication::Get().GetKeyboardFocusedWidget())
+		{
+			OutWidgets.Add(Widget.ToSharedRef());
+		}
+	});
+}
+
+TSharedRef<IElementLocator, ESPMode::ThreadSafe> By::UserFocus(uint32 UserIndex)
+{
+	return WidgetLambda([UserIndex](TArray<TSharedRef<SWidget>>& OutWidgets) -> void {
+		if (const TSharedPtr<SWidget> Widget = FSlateApplication::Get().GetUserFocusedWidget(UserIndex))
+		{
+			OutWidgets.Add(Widget.ToSharedRef());
+		}
+	});
+}
+
 TSharedRef<IElementLocator, ESPMode::ThreadSafe> By::TextFilter::Contains(
 	const FElementLocatorRef& RootLocator,
 	const FString& Value,
diff --git a/Engine/Source/Developer/AutomationDriver/Private/WaitUntil.cpp b/Engine/Source/Developer/AutomationDriver/Private/WaitUntil.cpp
--- a/Engine/Source/Developer/AutomationDriver/Private/WaitUntil.cpp
+++ b/Engine/Source/Developer/AutomationDriver/Private/WaitUntil.cpp
@@ -305,6 +305,62 @@ FDriverWaitDelegate Until::ElementIsScrolledToEnd(const TSharedRef<IElementLocat
 	});
 }
 
+FDriverWaitDelegate Until::ElementIsFocusedByKeyboard(const TSharedRef<IElementLocator, ESPMode::ThreadSafe>& ElementLocator,
+	FWaitTimeout Timeout)
+{
+	return Until::ElementIsFocusedByKeyboard(ElementLocator, FWaitInterval::InSeconds(1.0), Timeout);
+}
+
+FDriverWaitDelegate Until::ElementIsFocusedByKeyboard(const TSharedRef<IElementLocator, ESPMode::ThreadSafe>& ElementLocator,
+	FWaitInterval Interval, FWaitTimeout Timeout)
+{
+	return Until::Condition([ElementLocator]()
+	{
+		check(IsInGameThread());
+
+		TArray<TSharedRef<IApplicationElement>> Elements;
+		ElementLocator->Locate(Elements);
+
+		for (const TSharedRef<IApplicationElement>& Element : Elements)
+		{
+			if (Element->IsFocused())
+			{
+				return true;
+			}
+		}
+
+		return false;
+	}, Interval, Timeout);
+}
+
+FDriverWaitDelegate Until::ElementIsFocusedByUser(uint32 UserIndex,
+	const TSharedRef<IElementLocator, ESPMode::ThreadSafe>& ElementLocator, FWaitTimeout Timeout)
+{
+	return Until::ElementIsFocusedByUser(UserIndex, ElementLocator, FWaitInterval::InSeconds(1.0), Timeout);
+}
+
+FDriverWaitDelegate Until::ElementIsFocusedByUser(uint32 UserIndex,
+	const TSharedRef<IElementLocator, ESPMode::ThreadSafe>& ElementLocator, FWaitInterval Interval, FWaitTimeout Timeout)
+{
+	return Until::Condition([ElementLocator, UserIndex]()
+	{
+		check(IsInGameThread());
+
+		TArray<TSharedRef<IApplicationElement>> Elements;
+		ElementLocator->Locate(Elements);
+
+		for (const TSharedRef<IApplicationElement>& Element : Elements)
+		{
+			if (Element->IsFocused(UserIndex))
+			{
+				return true;
+			}
+		}
+
+		return false;
+	}, Interval, Timeout);
+}
+
 FDriverWaitDelegate Until::Condition(const TFunction<bool()>& Function, FWaitTimeout Timeout)
 {
 	return Until::Condition(FDriverWaitConditionDelegate::CreateLambda(Function), Timeout);
diff --git a/Engine/Source/Developer/AutomationDriver/Public/LocateBy.h b/Engine/Source/Developer/AutomationDriver/Public/LocateBy.h
--- a/Engine/Source/Developer/AutomationDriver/Public/LocateBy.h
+++ b/Engine/Source/Developer/AutomationDriver/Public/LocateBy.h
@@ -21,33 +21,37 @@ public:
 	 * Creates a new element locator that exposes the collection of SWidgets returned from the FLocateSlateWidgetElementDelegate
 	 * as discovered elements
 	 * @param Value - The delegate to use
+	 * @param DebugName - The string to display for errors when no elements can be located
 	 * @return a locator which uses the specified delegate to discover appropriate elements
 	 */
-	static TSharedRef<IElementLocator, ESPMode::ThreadSafe> Delegate(const FLocateSlateWidgetElementDelegate& Value);
+	static TSharedRef<IElementLocator, ESPMode::ThreadSafe> Delegate(const FLocateSlateWidgetElementDelegate& Value, FStringView DebugName = { });
 
 	/**
 	 * Creates a new element locator that exposes the collection of FWidgetPaths returned from the FLocateSlateWidgetElementDelegate
 	 * as discovered elements
 	 * @param Value - The delegate to use
+	 * @param DebugName - The string to display for errors when no elements can be located
 	 * @return a locator which uses the specified delegate to discover appropriate elements
 	 */
-	static TSharedRef<IElementLocator, ESPMode::ThreadSafe> Delegate(const FLocateSlateWidgetPathElementDelegate& Value);
+	static TSharedRef<IElementLocator, ESPMode::ThreadSafe> Delegate(const FLocateSlateWidgetPathElementDelegate& Value, FStringView DebugName = { });
 
 	/**
 	 * Creates a new element locator that exposes the collection of SWidgets returned from the lambda
 	 * as discovered elements
 	 * @param Value - The lambda to use
+	 * @param DebugName - The string to display for errors when no elements can be located
 	 * @return a locator which uses the specified lambda to discover appropriate elements
 	 */
-	static TSharedRef<IElementLocator, ESPMode::ThreadSafe> WidgetLambda(const TFunction<void(TArray<TSharedRef<SWidget>>&)>& Value);
+	static TSharedRef<IElementLocator, ESPMode::ThreadSafe> WidgetLambda(const TFunction<void(TArray<TSharedRef<SWidget>>&)>& Value, FStringView DebugName = { });
 
 	/**
 	 * Creates a new element locator that exposes the collection of FWidgetPaths returned from the lambda
 	 * as discovered elements
 	 * @param Value - The lambda to use
+	 * @param DebugName - The string to display for errors when no elements can be located
 	 * @return a locator which uses the specified lambda to discover appropriate elements
 	 */
-	static TSharedRef<IElementLocator, ESPMode::ThreadSafe> WidgetPathLambda(const TFunction<void(TArray<FWidgetPath>&)>& Value);
+	static TSharedRef<IElementLocator, ESPMode::ThreadSafe> WidgetPathLambda(const TFunction<void(TArray<FWidgetPath>&)>& Value, FStringView DebugName = { });
 
 	/**
 	 * Creates a new element locator that limits its discovered elements to ones with the specified Id
@@ -314,6 +318,18 @@ public:
 	 */
 	static TSharedRef<IElementLocator, ESPMode::ThreadSafe> Cursor();
 	
+	/**
+	 * Creates a new element locator that limits its discovered element to the one that's currently in keyboard focus
+	 * @return a locator which uses the keyboard focus to discover the appropriate element
+	 */
+	static TSharedRef<IElementLocator, ESPMode::ThreadSafe> KeyboardFocus();
+
+	/**
+	 * Creates a new element locator that limits its discovered element to the one that's currently in user focus
+	 * @return a locator which uses the user focus to discover the appropriate element
+	 */
+	static TSharedRef<IElementLocator, ESPMode::ThreadSafe> UserFocus(uint32 UserIndex = 0);
+
 	class AUTOMATIONDRIVER_API TextFilter
 	{
 	public:
diff --git a/Engine/Source/Developer/AutomationDriver/Public/WaitUntil.h b/Engine/Source/Developer/AutomationDriver/Public/WaitUntil.h
--- a/Engine/Source/Developer/AutomationDriver/Public/WaitUntil.h
+++ b/Engine/Source/Developer/AutomationDriver/Public/WaitUntil.h
@@ -167,6 +167,28 @@ public:
 	 */
 	static FDriverWaitDelegate ElementIsScrolledToEnd(const TSharedRef<IElementLocator, ESPMode::ThreadSafe>& ElementLocator, FWaitInterval Interval, FWaitTimeout Timeout);
 
+	/**
+	 * Creates a new wait delegate which completes its wait only if the specified element locator discovers an element focused by the default user or if the specified timeout timespan elapses
+	 */
+	static FDriverWaitDelegate ElementIsFocusedByKeyboard(const TSharedRef<IElementLocator, ESPMode::ThreadSafe>& ElementLocator, FWaitTimeout Timeout);
+
+	/**
+	 * Creates a new wait delegate which completes it's wait only if the specified element locator discovers an element focused by the default user or if the specified timeout timespan elapses;
+	 * The element locator is only re-evaluated at the specified wait interval
+	 */
+	static FDriverWaitDelegate ElementIsFocusedByKeyboard(const TSharedRef<IElementLocator, ESPMode::ThreadSafe>& ElementLocator, FWaitInterval Interval, FWaitTimeout Timeout);
+
+	/**
+	 * Creates a new wait delegate which completes its wait only if the specified element locator discovers an element focused by the specified user or if the specified timeout timespan elapses
+	 */
+	static FDriverWaitDelegate ElementIsFocusedByUser(uint32 UserIndex, const TSharedRef<IElementLocator, ESPMode::ThreadSafe>& ElementLocator, FWaitTimeout Timeout);
+
+	/**
+	 * Creates a new wait delegate which completes it's wait only if the specified element locator discovers an element focused by the specified user or if the specified timeout timespan elapses;
+	 * The element locator is only re-evaluated at the specified wait interval
+	 */
+	static FDriverWaitDelegate ElementIsFocusedByUser(uint32 UserIndex, const TSharedRef<IElementLocator, ESPMode::ThreadSafe>& ElementLocator, FWaitInterval Interval, FWaitTimeout Timeout);
+
 	/**
 	 * Creates a new wait delegate which completes it's wait only if the specified condition returns true or if the specified timeout timespan elapses
 	 */
