[Slate] Remove deprecated "OnControllerButtonXXX" functions, which were deprecated in 5.1.

These functions are deprecated because they use the raw int32 controller id, and we now prefer the versions which take in explicit FPlatformUserId and FInputDeviceIds.

diff --git a/Engine/Plugins/Experimental/RemoteSession/Source/RemoteSession/Private/MessageHandler/ProxyMessageHandler.cpp b/Engine/Plugins/Experimental/RemoteSession/Source/RemoteSession/Private/MessageHandler/ProxyMessageHandler.cpp
--- a/Engine/Plugins/Experimental/RemoteSession/Source/RemoteSession/Private/MessageHandler/ProxyMessageHandler.cpp
+++ b/Engine/Plugins/Experimental/RemoteSession/Source/RemoteSession/Private/MessageHandler/ProxyMessageHandler.cpp
@@ -173,38 +173,6 @@ bool FProxyMessageHandler::OnCursorSet()
 	return false;
 }
 
-PRAGMA_DISABLE_DEPRECATION_WARNINGS
-bool FProxyMessageHandler::OnControllerAnalog(FGamepadKeyNames::Type KeyName, int32 ControllerId, float AnalogValue)
-{
-	if (TargetHandler.IsValid())
-	{
-		return TargetHandler->OnControllerAnalog(KeyName, ControllerId, AnalogValue);
-	}
-
-	return false;
-}
-
-bool FProxyMessageHandler::OnControllerButtonPressed(FGamepadKeyNames::Type KeyName, int32 ControllerId, bool IsRepeat)
-{
-	if (TargetHandler.IsValid())
-	{
-		return TargetHandler->OnControllerButtonPressed(KeyName, ControllerId, IsRepeat);
-	}
-
-	return false;
-}
-
-bool FProxyMessageHandler::OnControllerButtonReleased(FGamepadKeyNames::Type KeyName, int32 ControllerId, bool IsRepeat)
-{
-	if (TargetHandler.IsValid())
-	{
-		return TargetHandler->OnControllerButtonReleased(KeyName, ControllerId, IsRepeat);
-	}
-
-	return false;
-}
-PRAGMA_ENABLE_DEPRECATION_WARNINGS
-
 bool FProxyMessageHandler::OnControllerAnalog(FGamepadKeyNames::Type KeyName, FPlatformUserId PlatformUserId, FInputDeviceId InputDeviceId, float AnalogValue)
 {
 	if (TargetHandler.IsValid())
diff --git a/Engine/Plugins/Experimental/RemoteSession/Source/RemoteSession/Private/MessageHandler/ProxyMessageHandler.h b/Engine/Plugins/Experimental/RemoteSession/Source/RemoteSession/Private/MessageHandler/ProxyMessageHandler.h
--- a/Engine/Plugins/Experimental/RemoteSession/Source/RemoteSession/Private/MessageHandler/ProxyMessageHandler.h
+++ b/Engine/Plugins/Experimental/RemoteSession/Source/RemoteSession/Private/MessageHandler/ProxyMessageHandler.h
@@ -12,11 +12,11 @@ class FProxyMessageHandler : public FGenericApplicationMessageHandler
 	
 public:
 
-	FProxyMessageHandler();
+	FProxyMessageHandler() = default;
 
 	FProxyMessageHandler(const TSharedPtr<FGenericApplicationMessageHandler>& InTargetHandler);
 
-	virtual ~FProxyMessageHandler();
+	virtual ~FProxyMessageHandler() override;
 
 	void SetTargetHandler(const TSharedPtr<FGenericApplicationMessageHandler>& InTargetHandler);
 
@@ -52,15 +52,6 @@ public:
 
 	virtual bool OnCursorSet() override;
 
-	UE_DEPRECATED(5.1, "This version of OnControllerAnalog has been deprecated, please use the one that takes an FPlatformUser and FInputDeviceId instead.")
-	virtual bool OnControllerAnalog(FGamepadKeyNames::Type KeyName, int32 ControllerId, float AnalogValue) override;
-
-	UE_DEPRECATED(5.1, "This version of OnControllerButtonPressed has been deprecated, please use the one that takes an FPlatformUser and FInputDeviceId instead.")
-	virtual bool OnControllerButtonPressed(FGamepadKeyNames::Type KeyName, int32 ControllerId, bool IsRepeat) override;
-
-	UE_DEPRECATED(5.1, "This version of OnControllerButtonReleased has been deprecated, please use the one that takes an FPlatformUser and FInputDeviceId instead.")
-	virtual bool OnControllerButtonReleased(FGamepadKeyNames::Type KeyName, int32 ControllerId, bool IsRepeat) override;
-
 	virtual bool OnControllerAnalog(FGamepadKeyNames::Type KeyName, FPlatformUserId PlatformUserId, FInputDeviceId InputDeviceId, float AnalogValue) override;
 
 	virtual bool OnControllerButtonPressed(FGamepadKeyNames::Type KeyName, FPlatformUserId PlatformUserId, FInputDeviceId InputDeviceId, bool IsRepeat) override;
diff --git a/Engine/Plugins/Experimental/RemoteSession/Source/RemoteSession/Private/MessageHandler/RecordingMessageHandler.cpp b/Engine/Plugins/Experimental/RemoteSession/Source/RemoteSession/Private/MessageHandler/RecordingMessageHandler.cpp
--- a/Engine/Plugins/Experimental/RemoteSession/Source/RemoteSession/Private/MessageHandler/RecordingMessageHandler.cpp
+++ b/Engine/Plugins/Experimental/RemoteSession/Source/RemoteSession/Private/MessageHandler/RecordingMessageHandler.cpp
@@ -651,70 +651,41 @@ void FRecordingMessageHandler::PlayOnMotionDetected(FArchive& Ar)
 }
 
 PRAGMA_DISABLE_DEPRECATION_WARNINGS
-bool FRecordingMessageHandler::OnControllerAnalog(FGamepadKeyNames::Type KeyName, int32 ControllerId, float AnalogValue)
-{
-	if (IsRecording())
-	{
-		ThreeParamMsg<FString, int32, float> Msg(KeyName.ToString(), ControllerId, AnalogValue);
-		RecordMessage(TEXT("OnControllerAnalog"), Msg.AsData());
-	}
-
-	if (bConsumeInput)
-	{
-		return true;
-	}
-
-	return FProxyMessageHandler::OnControllerAnalog(KeyName, ControllerId, AnalogValue);
-}
 
 void FRecordingMessageHandler::PlayOnControllerAnalog(FArchive& Ar)
 {
 	ThreeParamMsg<FString, int32, float > Msg(Ar);
-	OnControllerAnalog(FName(*Msg.Param1), Msg.Param2, Msg.Param3);
-}
-
-bool FRecordingMessageHandler::OnControllerButtonPressed(FGamepadKeyNames::Type KeyName, int32 ControllerId, bool IsRepeat)
-{
-	if (IsRecording())
-	{
-		ThreeParamMsg<FString, int32, bool> Msg(KeyName.ToString(), ControllerId, IsRepeat);
-		RecordMessage(TEXT("OnControllerButtonPressed"), Msg.AsData());
-	}
-
-	if (bConsumeInput)
-	{
-		return true;
-	}
-
-	return FProxyMessageHandler::OnControllerButtonPressed(KeyName, ControllerId, IsRepeat);
+	
+	const int32 LegacyControllerId = Msg.Param2;
+	FPlatformUserId PlatformUser = PLATFORMUSERID_NONE;
+	FInputDeviceId DeviceId = INPUTDEVICEID_NONE;
+	IPlatformInputDeviceMapper::Get().RemapControllerIdToPlatformUserAndDevice(LegacyControllerId, PlatformUser, DeviceId);
+	
+	OnControllerAnalog(FName(*Msg.Param1), PlatformUser, DeviceId, Msg.Param3);
 }
 
 void FRecordingMessageHandler::PlayOnControllerButtonPressed(FArchive& Ar)
 {
 	ThreeParamMsg<FString, int32, bool > Msg(Ar);
-	OnControllerButtonPressed(FName(*Msg.Param1), Msg.Param2, Msg.Param3);
-}
 
-bool FRecordingMessageHandler::OnControllerButtonReleased(FGamepadKeyNames::Type KeyName, int32 ControllerId, bool IsRepeat)
-{
-	if (IsRecording())
-	{
-		ThreeParamMsg<FString, int32, bool> Msg(KeyName.ToString(), ControllerId, IsRepeat);
-		RecordMessage(TEXT("OnControllerButtonReleased"), Msg.AsData());
-	}
-
-	if (bConsumeInput)
-	{
-		return true;
-	}
-
-	return FProxyMessageHandler::OnControllerButtonReleased(KeyName, ControllerId, IsRepeat);
+	const int32 LegacyControllerId = Msg.Param2;
+	FPlatformUserId PlatformUser = PLATFORMUSERID_NONE;
+	FInputDeviceId DeviceId = INPUTDEVICEID_NONE;
+	IPlatformInputDeviceMapper::Get().RemapControllerIdToPlatformUserAndDevice(LegacyControllerId, PlatformUser, DeviceId);
+	
+	OnControllerButtonPressed(FName(*Msg.Param1), PlatformUser, DeviceId, Msg.Param3);
 }
 
 void FRecordingMessageHandler::PlayOnControllerButtonReleased(FArchive& Ar)
 {
 	ThreeParamMsg<FString, int32, bool > Msg(Ar);
-	OnControllerButtonReleased(FName(*Msg.Param1), Msg.Param2, Msg.Param3);
+
+	const int32 LegacyControllerId = Msg.Param2;
+	FPlatformUserId PlatformUser = PLATFORMUSERID_NONE;
+	FInputDeviceId DeviceId = INPUTDEVICEID_NONE;
+	IPlatformInputDeviceMapper::Get().RemapControllerIdToPlatformUserAndDevice(LegacyControllerId, PlatformUser, DeviceId);
+	
+	OnControllerButtonReleased(FName(*Msg.Param1), PlatformUser, DeviceId, Msg.Param3);
 }
 PRAGMA_ENABLE_DEPRECATION_WARNINGS
 
diff --git a/Engine/Plugins/Experimental/RemoteSession/Source/RemoteSession/Private/MessageHandler/RecordingMessageHandler.h b/Engine/Plugins/Experimental/RemoteSession/Source/RemoteSession/Private/MessageHandler/RecordingMessageHandler.h
--- a/Engine/Plugins/Experimental/RemoteSession/Source/RemoteSession/Private/MessageHandler/RecordingMessageHandler.h
+++ b/Engine/Plugins/Experimental/RemoteSession/Source/RemoteSession/Private/MessageHandler/RecordingMessageHandler.h
@@ -130,13 +130,6 @@ public:
 	/**
 	 * Controller handling
 	 */
-	UE_DEPRECATED(5.1, "This version of OnControllerAnalog has been deprecated, please use the one that takes an FPlatformUser and FInputDeviceId instead.")
-	virtual bool OnControllerAnalog(FGamepadKeyNames::Type KeyName, int32 ControllerId, float AnalogValue) override;
-	UE_DEPRECATED(5.1, "This version of OnControllerButtonPressed has been deprecated, please use the one that takes an FPlatformUser and FInputDeviceId instead.")
-	virtual bool OnControllerButtonPressed(FGamepadKeyNames::Type KeyName, int32 ControllerId, bool IsRepeat) override;
-	UE_DEPRECATED(5.1, "This version of OnControllerButtonReleased has been deprecated, please use the one that takes an FPlatformUser and FInputDeviceId instead.")
-	virtual bool OnControllerButtonReleased(FGamepadKeyNames::Type KeyName, int32 ControllerId, bool IsRepeat) override;
-
 	virtual bool OnControllerAnalog(FGamepadKeyNames::Type KeyName, FPlatformUserId PlatformUserId, FInputDeviceId InputDeviceId, float AnalogValue) override;
 	virtual bool OnControllerButtonPressed(FGamepadKeyNames::Type KeyName, FPlatformUserId PlatformUserId, FInputDeviceId InputDeviceId, bool IsRepeat) override;
 	virtual bool OnControllerButtonReleased(FGamepadKeyNames::Type KeyName, FPlatformUserId PlatformUserId, FInputDeviceId InputDeviceId, bool IsRepeat) override;
diff --git a/Engine/Plugins/Experimental/ScreenReader/Source/ScreenReader/Private/GenericPlatform/ScreenReaderApplicationMessageHandlerBase.cpp b/Engine/Plugins/Experimental/ScreenReader/Source/ScreenReader/Private/GenericPlatform/ScreenReaderApplicationMessageHandlerBase.cpp
--- a/Engine/Plugins/Experimental/ScreenReader/Source/ScreenReader/Private/GenericPlatform/ScreenReaderApplicationMessageHandlerBase.cpp
+++ b/Engine/Plugins/Experimental/ScreenReader/Source/ScreenReader/Private/GenericPlatform/ScreenReaderApplicationMessageHandlerBase.cpp
@@ -87,23 +87,6 @@ bool FScreenReaderApplicationMessageHandlerBase::OnCursorSet()
 	return TargetMessageHandler->OnCursorSet();
 }
 
-PRAGMA_DISABLE_DEPRECATION_WARNINGS
-bool FScreenReaderApplicationMessageHandlerBase::OnControllerAnalog(FGamepadKeyNames::Type KeyName, int32 ControllerId, float AnalogValue)
-{
-	return TargetMessageHandler->OnControllerAnalog(KeyName, ControllerId, AnalogValue);
-}
-
-bool FScreenReaderApplicationMessageHandlerBase::OnControllerButtonPressed(FGamepadKeyNames::Type KeyName, int32 ControllerId, bool IsRepeat)
-{
-	return TargetMessageHandler->OnControllerButtonPressed(KeyName, ControllerId, IsRepeat);
-}
-
-bool FScreenReaderApplicationMessageHandlerBase::OnControllerButtonReleased(FGamepadKeyNames::Type KeyName, int32 ControllerId, bool IsRepeat)
-{
-	return TargetMessageHandler->OnControllerButtonReleased(KeyName, ControllerId, IsRepeat);
-}
-PRAGMA_ENABLE_DEPRECATION_WARNINGS
-
 bool FScreenReaderApplicationMessageHandlerBase::OnControllerAnalog(FGamepadKeyNames::Type KeyName, FPlatformUserId PlatformUserId, FInputDeviceId InputDeviceId, float AnalogValue)
 {
 	return TargetMessageHandler->OnControllerAnalog(KeyName, PlatformUserId, InputDeviceId, AnalogValue);
diff --git a/Engine/Plugins/Experimental/ScreenReader/Source/ScreenReader/Public/GenericPlatform/ScreenReaderApplicationMessageHandlerBase.h b/Engine/Plugins/Experimental/ScreenReader/Source/ScreenReader/Public/GenericPlatform/ScreenReaderApplicationMessageHandlerBase.h
--- a/Engine/Plugins/Experimental/ScreenReader/Source/ScreenReader/Public/GenericPlatform/ScreenReaderApplicationMessageHandlerBase.h
+++ b/Engine/Plugins/Experimental/ScreenReader/Source/ScreenReader/Public/GenericPlatform/ScreenReaderApplicationMessageHandlerBase.h
@@ -53,18 +53,9 @@ public:
 	virtual bool OnMouseMove() override;
 	virtual bool OnRawMouseMove(const int32 X, const int32 Y) override;
 	virtual bool OnCursorSet() override;
-
-	UE_DEPRECATED(5.1, "This version of OnControllerAnalog is deprecated, please use the version that takes an FPlatformUserId instead")
-	virtual bool OnControllerAnalog(FGamepadKeyNames::Type KeyName, int32 ControllerId, float AnalogValue) override;
-	UE_DEPRECATED(5.1, "This version of OnControllerButtonPressed is deprecated, please use the version that takes an FPlatformUserId instead")
-	virtual bool OnControllerButtonPressed(FGamepadKeyNames::Type KeyName, int32 ControllerId, bool IsRepeat) override;
-	UE_DEPRECATED(5.1, "This version of OnControllerButtonReleased is deprecated, please use the version that takes an FPlatformUserId instead")
-	virtual bool OnControllerButtonReleased(FGamepadKeyNames::Type KeyName, int32 ControllerId, bool IsRepeat) override;
-
 	virtual bool OnControllerAnalog(FGamepadKeyNames::Type KeyName, FPlatformUserId PlatformUserId, FInputDeviceId InputDeviceId, float AnalogValue) override;
 	virtual bool OnControllerButtonPressed(FGamepadKeyNames::Type KeyName, FPlatformUserId PlatformUserId, FInputDeviceId InputDeviceId, bool IsRepeat) override;
 	virtual bool OnControllerButtonReleased(FGamepadKeyNames::Type KeyName, FPlatformUserId PlatformUserId, FInputDeviceId InputDeviceId, bool IsRepeat) override;
-	
 	virtual void OnBeginGesture() override;
 	virtual bool OnTouchGesture(EGestureEvent GestureType, const FVector2D& Delta, float WheelDelta, bool bIsDirectionInvertedFromDevice) override;
 	virtual void OnEndGesture() override;
diff --git a/Engine/Source/Developer/AutomationDriver/Private/PassThroughMessageHandler.cpp b/Engine/Source/Developer/AutomationDriver/Private/PassThroughMessageHandler.cpp
--- a/Engine/Source/Developer/AutomationDriver/Private/PassThroughMessageHandler.cpp
+++ b/Engine/Source/Developer/AutomationDriver/Private/PassThroughMessageHandler.cpp
@@ -10,7 +10,7 @@ class FPassThroughMessageHandlerImpl
 {
 public:
 
-	virtual ~FPassThroughMessageHandlerImpl()
+	virtual ~FPassThroughMessageHandlerImpl() override
 	{ }
 
 	virtual bool IsHandlingMessages() const override
@@ -78,7 +78,7 @@ public:
 		return RealMessageHandler->OnMouseDown(Window, Button);
 	}
 
-	virtual bool OnMouseDown(const TSharedPtr< FGenericWindow >& Window, const EMouseButtons::Type Button, const FVector2D CursorPos)
+	virtual bool OnMouseDown(const TSharedPtr< FGenericWindow >& Window, const EMouseButtons::Type Button, const FVector2D CursorPos) override
 	{
 		if (!bAllowMessageHandling)
 		{
@@ -98,7 +98,7 @@ public:
 		return RealMessageHandler->OnMouseUp(Button);
 	}
 
-	virtual bool OnMouseUp(const EMouseButtons::Type Button, const FVector2D CursorPos)
+	virtual bool OnMouseUp(const EMouseButtons::Type Button, const FVector2D CursorPos) override
 	{
 		if (!bAllowMessageHandling)
 		{
@@ -118,7 +118,7 @@ public:
 		return RealMessageHandler->OnMouseDoubleClick(Window, Button);
 	}
 
-	virtual bool OnMouseDoubleClick(const TSharedPtr< FGenericWindow >& Window, const EMouseButtons::Type Button, const FVector2D CursorPos)
+	virtual bool OnMouseDoubleClick(const TSharedPtr< FGenericWindow >& Window, const EMouseButtons::Type Button, const FVector2D CursorPos) override
 	{
 		if (!bAllowMessageHandling)
 		{
@@ -138,7 +138,7 @@ public:
 		return RealMessageHandler->OnMouseWheel(Delta);
 	}
 
-	virtual bool OnMouseWheel(const float Delta, const FVector2D CursorPos)
+	virtual bool OnMouseWheel(const float Delta, const FVector2D CursorPos) override
 	{
 		if (!bAllowMessageHandling)
 		{
@@ -177,47 +177,8 @@ public:
 
 		return RealMessageHandler->OnCursorSet();
 	}
-
-	UE_DEPRECATED(5.1, "This version of OnControllerAnalog is deprecated, please use the version that takes an FPlatformUserId instead")
-	virtual bool OnControllerAnalog(FGamepadKeyNames::Type KeyName, int32 ControllerId, float AnalogValue) override
-	{
-		if (!bAllowMessageHandling)
-		{
-			return false;
-		}
-
-		PRAGMA_DISABLE_DEPRECATION_WARNINGS
-		return RealMessageHandler->OnControllerAnalog(KeyName, ControllerId, AnalogValue);
-		PRAGMA_ENABLE_DEPRECATION_WARNINGS
-	}
-
-	UE_DEPRECATED(5.1, "This version of OnControllerAnalog is deprecated, please use the version that takes an FPlatformUserId instead")
-	virtual bool OnControllerButtonPressed(FGamepadKeyNames::Type KeyName, int32 ControllerId, bool IsRepeat) override
-	{
-		if (!bAllowMessageHandling)
-		{
-			return false;
-		}
-
-		PRAGMA_DISABLE_DEPRECATION_WARNINGS
-		return RealMessageHandler->OnControllerButtonPressed(KeyName, ControllerId, IsRepeat);
-		PRAGMA_ENABLE_DEPRECATION_WARNINGS
-	}
-
-	UE_DEPRECATED(5.1, "This version of OnControllerAnalog is deprecated, please use the version that takes an FPlatformUserId instead")
-	virtual bool OnControllerButtonReleased(FGamepadKeyNames::Type KeyName, int32 ControllerId, bool IsRepeat) override
-	{
-		if (!bAllowMessageHandling)
-		{
-			return false;
-		}
-
-		PRAGMA_DISABLE_DEPRECATION_WARNINGS
-		return RealMessageHandler->OnControllerButtonReleased(KeyName, ControllerId, IsRepeat);
-		PRAGMA_ENABLE_DEPRECATION_WARNINGS
-	}
 	
-	virtual bool OnControllerAnalog(FGamepadKeyNames::Type KeyName, FPlatformUserId PlatformUserId, FInputDeviceId InputDeviceId, float AnalogValue)
+	virtual bool OnControllerAnalog(FGamepadKeyNames::Type KeyName, FPlatformUserId PlatformUserId, FInputDeviceId InputDeviceId, float AnalogValue) override
 	{
 		if (!bAllowMessageHandling)
 		{
@@ -227,7 +188,7 @@ public:
 		return RealMessageHandler->OnControllerAnalog(KeyName, PlatformUserId, InputDeviceId, AnalogValue);
 	}
 
-	virtual bool OnControllerButtonPressed(FGamepadKeyNames::Type KeyName, FPlatformUserId PlatformUserId, FInputDeviceId InputDeviceId, bool IsRepeat)
+	virtual bool OnControllerButtonPressed(FGamepadKeyNames::Type KeyName, FPlatformUserId PlatformUserId, FInputDeviceId InputDeviceId, bool IsRepeat) override
 	{
 		if (!bAllowMessageHandling)
 		{
@@ -237,7 +198,7 @@ public:
 		return RealMessageHandler->OnControllerButtonPressed(KeyName, PlatformUserId, InputDeviceId, IsRepeat);
 	}
 
-	virtual bool OnControllerButtonReleased(FGamepadKeyNames::Type KeyName, FPlatformUserId PlatformUserId, FInputDeviceId InputDeviceId, bool IsRepeat)
+	virtual bool OnControllerButtonReleased(FGamepadKeyNames::Type KeyName, FPlatformUserId PlatformUserId, FInputDeviceId InputDeviceId, bool IsRepeat) override
 	{
 		if (!bAllowMessageHandling)
 		{
@@ -519,7 +480,7 @@ class FPassThroughMessageHandlerFactory
 {
 public:
 
-	virtual ~FPassThroughMessageHandlerFactory()
+	virtual ~FPassThroughMessageHandlerFactory() override
 	{ }
 
 	virtual TSharedRef<FPassThroughMessageHandler> Create(
diff --git a/Engine/Source/Editor/Persona/Private/AnimationEditorViewportClient.cpp b/Engine/Source/Editor/Persona/Private/AnimationEditorViewportClient.cpp
--- a/Engine/Source/Editor/Persona/Private/AnimationEditorViewportClient.cpp
+++ b/Engine/Source/Editor/Persona/Private/AnimationEditorViewportClient.cpp
@@ -1597,23 +1597,6 @@ bool FAnimationViewportClient::InputAxis(FViewport* InViewport, FInputDeviceId D
 	return bResult;
 }
 
-PRAGMA_DISABLE_DEPRECATION_WARNINGS
-bool FAnimationViewportClient::InputKey( FViewport* InViewport, int32 ControllerId, FKey Key, EInputEvent Event, float AmountDepressed, bool bGamepad )
-{
-	FInputKeyEventArgs Args(InViewport, ControllerId, Key, Event, AmountDepressed, /*bIsTouchEvent*/false);
-	return InputKey(Args);
-}
-
-bool FAnimationViewportClient::InputAxis(FViewport* InViewport, int32 ControllerId, FKey Key, float Delta, float DeltaTime, int32 NumSamples /*= 1*/, bool bGamepad /*= false*/)
-{
-	FPlatformUserId UserId = FGenericPlatformMisc::GetPlatformUserForUserIndex(ControllerId);
-	FInputDeviceId DeviceId = INPUTDEVICEID_NONE;
-	IPlatformInputDeviceMapper::Get().RemapControllerIdToPlatformUserAndDevice(ControllerId, UserId, DeviceId);
-	
-	return InputAxis(InViewport, DeviceId, Key, Delta, DeltaTime, NumSamples, bGamepad);
-}
-PRAGMA_ENABLE_DEPRECATION_WARNINGS
-
 void FAnimationViewportClient::SetLocalAxesMode(ELocalAxesMode::Type AxesMode)
 {
 	ConfigOption->SetDefaultLocalAxesSelection( AxesMode );
diff --git a/Engine/Source/Editor/Persona/Public/AnimationEditorViewportClient.h b/Engine/Source/Editor/Persona/Public/AnimationEditorViewportClient.h
--- a/Engine/Source/Editor/Persona/Public/AnimationEditorViewportClient.h
+++ b/Engine/Source/Editor/Persona/Public/AnimationEditorViewportClient.h
@@ -43,7 +43,7 @@ protected:
 
 public:
 	FAnimationViewportClient(const TSharedRef<class IPersonaPreviewScene>& InPreviewScene, const TSharedRef<class SAnimationEditorViewport>& InAnimationEditorViewport, const TSharedRef<class FAssetEditorToolkit>& InAssetEditorToolkit, int32 InViewportIndex, bool bInShowStats);
-	virtual ~FAnimationViewportClient();
+	virtual ~FAnimationViewportClient() override;
 
 	void Initialize();
 
@@ -51,13 +51,7 @@ public:
 	virtual void Tick(float DeltaSeconds) override;
 	virtual void Draw(const FSceneView* View,FPrimitiveDrawInterface* PDI) override;
 	virtual void DrawCanvas( FViewport& InViewport, FSceneView& View, FCanvas& Canvas ) override;
-	
-	UE_DEPRECATED(5.1, "This version of InputKey is deprecated. Please use the version that takes EventArgs instead.")
-	virtual bool InputKey(FViewport* Viewport, int32 ControllerId, FKey Key, EInputEvent Event, float AmountDepressed = 1.f, bool bGamepad=false) override;
 	virtual bool InputKey(const FInputKeyEventArgs& EventArgs) override;
-	
-	UE_DEPRECATED(5.1, "This version of InputAxis is deprecated. Please use the version that takes DeviceId instead.")
-	virtual bool InputAxis(FViewport* InViewport, int32 ControllerId, FKey Key, float Delta, float DeltaTime, int32 NumSamples = 1, bool bGamepad = false) override;
 	virtual bool InputAxis(FViewport* InViewport, FInputDeviceId DeviceId, FKey Key, float Delta, float DeltaTime, int32 NumSamples = 1, bool bGamepad = false) override;
 
 	// Sets what bones are drawn by DrawMeshBones and ShowBoneNames
diff --git a/Engine/Source/Editor/UnrealEd/Private/EditorViewportClient.cpp b/Engine/Source/Editor/UnrealEd/Private/EditorViewportClient.cpp
--- a/Engine/Source/Editor/UnrealEd/Private/EditorViewportClient.cpp
+++ b/Engine/Source/Editor/UnrealEd/Private/EditorViewportClient.cpp
@@ -3110,16 +3110,6 @@ bool FEditorViewportClient::ShouldScaleCameraSpeedByDistance() const
 	return GetDefault<ULevelEditorViewportSettings>()->bUseDistanceScaledCameraSpeed;
 }
 
-PRAGMA_DISABLE_DEPRECATION_WARNINGS
-bool FEditorViewportClient::InputKey(FViewport* InViewport, int32 ControllerId, FKey Key, EInputEvent Event, float AmountDepressed, bool bGamepad)
-{
-	FInputKeyEventArgs Args(InViewport, ControllerId, Key, Event);
-	Args.AmountDepressed = AmountDepressed;
-
-	return Internal_InputKey(Args);
-}
-PRAGMA_ENABLE_DEPRECATION_WARNINGS
-
 bool FEditorViewportClient::InputKey(const FInputKeyEventArgs& EventArgs)
 {
 	return Internal_InputKey(EventArgs);
@@ -4826,17 +4816,6 @@ FVector FEditorViewportClient::TranslateDelta( FKey InKey, float InDelta, bool I
 	return vec;
 }
 
-PRAGMA_DISABLE_DEPRECATION_WARNINGS
-bool FEditorViewportClient::InputAxis(FViewport* InViewport, int32 ControllerId, FKey Key, float Delta, float DeltaTime, int32 NumSamples, bool bGamepad)
-{
-	FInputDeviceId DeviceID = INPUTDEVICEID_NONE;
-	FPlatformUserId UserId = FGenericPlatformMisc::GetPlatformUserForUserIndex(ControllerId);
-	IPlatformInputDeviceMapper::Get().RemapControllerIdToPlatformUserAndDevice(ControllerId, UserId, DeviceID);
-	
-	return Internal_InputAxis(InViewport, DeviceID, Key, Delta, DeltaTime, NumSamples, bGamepad);
-}
-PRAGMA_ENABLE_DEPRECATION_WARNINGS
-
 bool FEditorViewportClient::InputAxis(FViewport* InViewport, FInputDeviceId DeviceID, FKey Key, float Delta, float DeltaTime, int32 NumSamples, bool bGamepad)
 {
 	return Internal_InputAxis(InViewport, DeviceID, Key, Delta, DeltaTime, NumSamples, bGamepad);
diff --git a/Engine/Source/Editor/UnrealEd/Public/EditorViewportClient.h b/Engine/Source/Editor/UnrealEd/Public/EditorViewportClient.h
--- a/Engine/Source/Editor/UnrealEd/Public/EditorViewportClient.h
+++ b/Engine/Source/Editor/UnrealEd/Public/EditorViewportClient.h
@@ -598,11 +598,7 @@ public:
 	UNREALED_API virtual bool ProcessScreenShots(FViewport* Viewport) override;
 	UNREALED_API virtual void RedrawRequested(FViewport* Viewport) override;
 	UNREALED_API virtual void RequestInvalidateHitProxy(FViewport* Viewport) override;
-	UE_DEPRECATED(5.1, "This version of InputKey is deprecated. Please use the version that takes FInputKeyEventArgs instead.")
-	UNREALED_API virtual bool InputKey(FViewport* InViewport, int32 ControllerId, FKey Key, EInputEvent Event, float AmountDepressed = 1.f, bool bGamepad=false) override;
 	UNREALED_API virtual bool InputKey(const FInputKeyEventArgs& EventArgs) override;
-	UE_DEPRECATED(5.1, "This version of InputAxis is deprecated. Please use the version that takes a DeviceId instead.")
-	UNREALED_API virtual bool InputAxis(FViewport* InViewport, int32 ControllerId, FKey Key, float Delta, float DeltaTime, int32 NumSamples=1, bool bGamepad=false) override;
 	UNREALED_API virtual bool InputAxis(FViewport* Viewport, FInputDeviceId DeviceID, FKey Key, float Delta, float DeltaTime, int32 NumSamples=1, bool bGamepad=false) override;
 	UNREALED_API virtual bool InputGesture(FViewport* Viewport, EGestureEvent GestureType, const FVector2D& GestureDelta, bool bIsDirectionInvertedFromDevice) override;
 	UNREALED_API virtual void ReceivedFocus(FViewport* Viewport) override;
diff --git a/Engine/Source/Runtime/ApplicationCore/Public/GenericPlatform/GenericApplicationMessageHandler.h b/Engine/Source/Runtime/ApplicationCore/Public/GenericPlatform/GenericApplicationMessageHandler.h
--- a/Engine/Source/Runtime/ApplicationCore/Public/GenericPlatform/GenericApplicationMessageHandler.h
+++ b/Engine/Source/Runtime/ApplicationCore/Public/GenericPlatform/GenericApplicationMessageHandler.h
@@ -204,8 +204,6 @@ public:
 	static APPLICATIONCORE_API const FInputDeviceScope* GetCurrent();
 };
 
-PRAGMA_DISABLE_DEPRECATION_WARNINGS
-
 /** Interface that defines how to handle interaction with a user via hardware input and output */
 class FGenericApplicationMessageHandler
 {
@@ -304,28 +302,16 @@ public:
 
 	virtual bool OnControllerAnalog(FGamepadKeyNames::Type KeyName, FPlatformUserId PlatformUserId, FInputDeviceId InputDeviceId, float AnalogValue)
 	{
-		if (!ShouldUsePlatformUserId())
-		{
-			return OnControllerAnalog(KeyName, PlatformUserId.GetInternalId(), AnalogValue);
-		}
 		return false;
 	}
 
 	virtual bool OnControllerButtonPressed(FGamepadKeyNames::Type KeyName, FPlatformUserId PlatformUserId, FInputDeviceId InputDeviceId, bool IsRepeat)
 	{
-		if (!ShouldUsePlatformUserId())
-		{
-			return OnControllerButtonPressed(KeyName, PlatformUserId.GetInternalId(), IsRepeat);
-		}
 		return false;
 	}
 
 	virtual bool OnControllerButtonReleased(FGamepadKeyNames::Type KeyName, FPlatformUserId PlatformUserId, FInputDeviceId InputDeviceId, bool IsRepeat)
 	{
-		if (!ShouldUsePlatformUserId())
-		{
-			return OnControllerButtonReleased(KeyName, PlatformUserId.GetInternalId(), IsRepeat);
-		}
 		return false;
 	}
 
@@ -510,47 +496,7 @@ public:
 	{
 
 	}
-
-	// Deprecate these when engine code has been converted to handle platform user id
-	UE_DEPRECATED(5.1, "This version of OnControllerAnalog has been deprecated, please use the one that takes an FPlatformUser and FInputDeviceId instead.")
-	virtual bool OnControllerAnalog(FGamepadKeyNames::Type KeyName, int32 ControllerId, float AnalogValue)
-	{
-		if (ShouldUsePlatformUserId())
-		{
-			// Remap the old int32 ControlerId to the new Platform user for backwards compat
-			FPlatformUserId UserId = FGenericPlatformMisc::GetPlatformUserForUserIndex(ControllerId);
-			FInputDeviceId DeviceId = INPUTDEVICEID_NONE;
-			IPlatformInputDeviceMapper::Get().RemapControllerIdToPlatformUserAndDevice(ControllerId, UserId, DeviceId);
-			return OnControllerAnalog(KeyName, UserId, DeviceId, AnalogValue);
-		}
-		return false;
-	}
-	UE_DEPRECATED(5.1, "This version of OnControllerButtonPressed has been deprecated, please use the one that takes an FPlatformUser and FInputDeviceId instead.")
-	virtual bool OnControllerButtonPressed(FGamepadKeyNames::Type KeyName, int32 ControllerId, bool IsRepeat)
-	{
-		if (ShouldUsePlatformUserId())
-		{
-			// Remap the old int32 ControlerId to the new Platform user for backwards compat
-			FPlatformUserId UserId = FGenericPlatformMisc::GetPlatformUserForUserIndex(ControllerId);
-			FInputDeviceId DeviceId = INPUTDEVICEID_NONE;
-			IPlatformInputDeviceMapper::Get().RemapControllerIdToPlatformUserAndDevice(ControllerId, UserId, DeviceId);
-			return OnControllerButtonPressed(KeyName, UserId, DeviceId, IsRepeat);
-		}
-		return false;
-	}
-	UE_DEPRECATED(5.1, "This version of OnControllerButtonReleased has been deprecated, please use the one that takes an FPlatformUser and FInputDeviceId instead.")
-	virtual bool OnControllerButtonReleased(FGamepadKeyNames::Type KeyName, int32 ControllerId, bool IsRepeat)
-	{
-		if (ShouldUsePlatformUserId())
-		{
-			// Remap the old int32 ControlerId to the new Platform user for backwards compat
-			FPlatformUserId UserId = FGenericPlatformMisc::GetPlatformUserForUserIndex(ControllerId);
-			FInputDeviceId DeviceId = INPUTDEVICEID_NONE;
-			IPlatformInputDeviceMapper::Get().RemapControllerIdToPlatformUserAndDevice(ControllerId, UserId, DeviceId);
-			return OnControllerButtonReleased(KeyName, UserId, DeviceId, IsRepeat);
-		}
-		return false;
-	}
+	
 	virtual bool OnTouchStarted(const TSharedPtr< FGenericWindow >& Window, const FVector2D& Location, float Force, int32 TouchIndex, int32 ControllerId)
 	{ 
 		if (ShouldUsePlatformUserId())
@@ -623,6 +569,4 @@ public:
 		}
 		return false;
 	}
-};
-
-PRAGMA_ENABLE_DEPRECATION_WARNINGS
+};
\ No newline at end of file
diff --git a/Engine/Source/Runtime/Engine/Classes/Engine/GameViewportClient.h b/Engine/Source/Runtime/Engine/Classes/Engine/GameViewportClient.h
--- a/Engine/Source/Runtime/Engine/Classes/Engine/GameViewportClient.h
+++ b/Engine/Source/Runtime/Engine/Classes/Engine/GameViewportClient.h
@@ -171,8 +171,6 @@ public:
 	virtual void RedrawRequested(FViewport* InViewport) override {}
 	ENGINE_API virtual bool InputKey(const FInputKeyEventArgs& EventArgs) override;
 	
-	UE_DEPRECATED(5.1, "This version of InputAxis has been deprecated. Please use the version that takes an FInputDeviceId instead")
-	ENGINE_API virtual bool InputAxis(FViewport* Viewport, int32 ControllerId, FKey Key, float Delta, float DeltaTime, int32 NumSamples=1, bool bGamepad=false) override;
 	ENGINE_API virtual bool InputAxis(FViewport* Viewport, FInputDeviceId InputDevice, FKey Key, float Delta, float DeltaTime, int32 NumSamples = 1, bool bGamepad = false) override;
 	
 	ENGINE_API virtual bool InputChar(FViewport* Viewport,int32 ControllerId, TCHAR Character) override;
diff --git a/Engine/Source/Runtime/Engine/Private/GameViewportClient.cpp b/Engine/Source/Runtime/Engine/Private/GameViewportClient.cpp
--- a/Engine/Source/Runtime/Engine/Private/GameViewportClient.cpp
+++ b/Engine/Source/Runtime/Engine/Private/GameViewportClient.cpp
@@ -763,19 +763,6 @@ bool UGameViewportClient::InputKey(const FInputKeyEventArgs& InEventArgs)
 	return bResult;
 }
 
-PRAGMA_DISABLE_DEPRECATION_WARNINGS
-bool UGameViewportClient::InputAxis(FViewport* InViewport, int32 ControllerId, FKey Key, float Delta, float DeltaTime, int32 NumSamples, bool bGamepad)
-{
-	// Remap the old int32 ControllerId value to the new InputDeviceId
-	IPlatformInputDeviceMapper& DeviceMapper = IPlatformInputDeviceMapper::Get();
-	FPlatformUserId UserId = PLATFORMUSERID_NONE;
-	FInputDeviceId DeviceId = INPUTDEVICEID_NONE;
-	DeviceMapper.RemapControllerIdToPlatformUserAndDevice(ControllerId, UserId, DeviceId);
-	
-	return InputAxis(InViewport, DeviceId, Key, Delta, DeltaTime, NumSamples, bGamepad);
-}
-PRAGMA_ENABLE_DEPRECATION_WARNINGS
-
 bool UGameViewportClient::InputAxis(FViewport* InViewport, FInputDeviceId InputDevice, FKey Key, float Delta, float DeltaTime, int32 NumSamples, bool bGamepad)
 {
 	if (IgnoreInput())
