Automate 'Animators - Remove an individual animator'

#rnx

diff --git a/Engine/Source/Developer/AutomationDriver/Private/DriverSequence.cpp b/Engine/Source/Developer/AutomationDriver/Private/DriverSequence.cpp
--- a/Engine/Source/Developer/AutomationDriver/Private/DriverSequence.cpp
+++ b/Engine/Source/Developer/AutomationDriver/Private/DriverSequence.cpp
@@ -235,6 +235,12 @@ public:
 		return *this;
 	}
 
+	virtual IAsyncActionSequence& MoveToPosition(const FVector2D& Position) override
+	{
+		InternalMoveToPosition(Position);
+		return *this;
+	}
+
 	virtual IAsyncActionSequence& ScrollBy(float Delta) override
 	{
 		InternalScrollBy(Delta);
@@ -931,6 +937,17 @@ private:
 		});
 	}
 
+	void InternalMoveToPosition(const FVector2D& Position)
+	{
+		StepsExecutor->Add([this, Position](const FTimespan& TotalProcessTime) -> FStepResult {
+
+			Application->Cursor->SetPosition(Position.X, Position.Y);
+			Application->GetRealMessageHandler()->OnMouseMove();
+
+			return FStep::Done();
+			});
+	}
+
 	void InternalScrollBy(float Delta)
 	{
 		StepsExecutor->Add([this, Delta](const FTimespan& TotalProcessTime) -> FStepResult {
@@ -1456,6 +1473,12 @@ public:
 		return *this;
 	}
 
+	virtual IActionSequence& MoveToPosition(const FVector2D& Position) override
+	{
+		ActionSequence->MoveToPosition(Position);
+		return *this;
+	}
+
 	virtual IActionSequence& ScrollBy(float Delta) override
 	{
 		ActionSequence->ScrollBy(Delta);
diff --git a/Engine/Source/Developer/AutomationDriver/Public/IDriverSequence.h b/Engine/Source/Developer/AutomationDriver/Public/IDriverSequence.h
--- a/Engine/Source/Developer/AutomationDriver/Public/IDriverSequence.h
+++ b/Engine/Source/Developer/AutomationDriver/Public/IDriverSequence.h
@@ -37,6 +37,11 @@ public:
 	 * Queues a move of the cursor from its current position by the specified offsets
 	 */
 	virtual IAsyncActionSequence& MoveByOffset(int32 XOffset, int32 YOffset) = 0;
+
+	/**
+	 * Queues a move of the cursor from its current position to a new position which is specified
+	 */
+	virtual IAsyncActionSequence& MoveToPosition(const FVector2D& Position) = 0;
 	
 	/**
 	 * Queues the series of events to invoke a mouse wheel scroll by the specified delta at the current mouse position
@@ -479,6 +484,11 @@ public:
 	 * Queues a move of the cursor from its current position by the specified offsets
 	 */
 	virtual IActionSequence& MoveByOffset(int32 XOffset, int32 YOffset) = 0;
+
+	/**
+	 * Queues a move of the cursor from its current position to a new position which is specified
+	 */
+	virtual IActionSequence& MoveToPosition(const FVector2D& Position) = 0;
 	
 	/**
 	 * Queues the series of events to invoke a mouse wheel scroll by the specified delta at the current mouse position
